Add RSA support to the PlatformCrypto plugin
 - NOTE: Only implemented for platforms that support OpenSSL. Other platforms have empty implementations.

diff --git a/Engine/Plugins/Experimental/PlatformCrypto/Source/PlatformCrypto/Public/PlatformCryptoTypes.h b/Engine/Plugins/Experimental/PlatformCrypto/Source/PlatformCrypto/Public/PlatformCryptoTypes.h
--- a/Engine/Plugins/Experimental/PlatformCrypto/Source/PlatformCrypto/Public/PlatformCryptoTypes.h
+++ b/Engine/Plugins/Experimental/PlatformCrypto/Source/PlatformCrypto/Public/PlatformCryptoTypes.h
@@ -14,6 +14,11 @@ enum class EPlatformCryptoResult
 	Failure
 };
 
+/**
+ * Generic handle to an RSA key
+ */
+typedef void* TPlatformCryptoRSAKey;
+
 /**
  * Stat group for implementations to use.
  */
diff --git a/Engine/Plugins/Experimental/PlatformCrypto/Source/PlatformCryptoOpenSSL/Private/EncryptionContextOpenSSL.cpp b/Engine/Plugins/Experimental/PlatformCrypto/Source/PlatformCryptoOpenSSL/Private/EncryptionContextOpenSSL.cpp
--- a/Engine/Plugins/Experimental/PlatformCrypto/Source/PlatformCryptoOpenSSL/Private/EncryptionContextOpenSSL.cpp
+++ b/Engine/Plugins/Experimental/PlatformCrypto/Source/PlatformCryptoOpenSSL/Private/EncryptionContextOpenSSL.cpp
@@ -173,3 +173,119 @@ bool FEncryptionContextOpenSSL::DigestVerify_PS256(const TArrayView<const char>
 	EVP_DigestVerifyUpdate(Context.Get(), Message.GetData(), Message.Num());
 	return EVP_DigestVerifyFinal(Context.Get(), const_cast<uint8*>(Signature.GetData()), Signature.Num()) == 1;
 }
+
+// Some platforms were upgraded to OpenSSL 1.1.1 while the others were left on a previous version. There are some minor differences we have to account for
+// in the older version, so declare a handy define that we can use to gate the code
+#if !defined(OPENSSL_VERSION_NUMBER) || OPENSSL_VERSION_NUMBER < 0x10100000L
+#define USE_LEGACY_OPENSSL 1
+#else
+#define USE_LEGACY_OPENSSL 0
+#endif
+
+static void LoadBinaryIntoBigNum(const uint8* InData, int64 InDataSize, BIGNUM* InBigNum)
+{
+#if USE_LEGACY_OPENSSL
+	TArray<uint8> Bytes(InData, InDataSize);
+	Algo::Reverse(Bytes);
+	BN_bin2bn(Bytes.GetData(), Bytes.Num(), InBigNum);
+#else
+	BN_lebin2bn(InData, InDataSize, InBigNum);
+#endif
+}
+
+TPlatformCryptoRSAKey FEncryptionContextOpenSSL::CreateKey_RSA(const TArrayView<const uint8> PublicExponent, const TArrayView<const uint8> PrivateExponent, const TArrayView<const uint8> Modulus)
+{
+	RSA* NewKey = RSA_new();
+
+	BIGNUM* BN_PublicExponent = PublicExponent.Num() > 0 ? BN_new() : nullptr;
+	BIGNUM* BN_PrivateExponent = PrivateExponent.Num() > 0 ? BN_new() : nullptr;
+	BIGNUM* BN_Modulus = BN_new();
+
+	if (PublicExponent.Num())
+	{
+		LoadBinaryIntoBigNum(PublicExponent.GetData(), PublicExponent.Num(), BN_PublicExponent);
+	}
+
+	if (PrivateExponent.Num())
+	{
+		LoadBinaryIntoBigNum(PrivateExponent.GetData(), PrivateExponent.Num(), BN_PrivateExponent);
+	}
+
+	LoadBinaryIntoBigNum(Modulus.GetData(), Modulus.Num(), BN_Modulus);
+#if USE_LEGACY_OPENSSL
+	NewKey->n = BN_Modulus;
+	NewKey->e = BN_PublicExponent;
+	NewKey->d = BN_PrivateExponent;
+#else
+	RSA_set0_key(NewKey, BN_Modulus, BN_PublicExponent, BN_PrivateExponent);
+#endif
+
+	return NewKey;
+}
+
+void FEncryptionContextOpenSSL::DestroyKey_RSA(TPlatformCryptoRSAKey Key)
+{
+	RSA_free((RSA*)Key);
+}
+
+int32 FEncryptionContextOpenSSL::GetKeySize_RSA(TPlatformCryptoRSAKey Key)
+{
+	return RSA_size((RSA*)Key);
+}
+
+int32 FEncryptionContextOpenSSL::GetMaxDataSize_RSA(TPlatformCryptoRSAKey Key)
+{
+	return (GetKeySize_RSA(Key)) - RSA_PKCS1_PADDING_SIZE;
+}
+
+int32 FEncryptionContextOpenSSL::EncryptPublic_RSA(TArrayView<const uint8> Source, TArray<uint8>& Dest, TPlatformCryptoRSAKey Key)
+{
+	Dest.SetNum(GetKeySize_RSA(Key));
+	int NumEncryptedBytes = RSA_public_encrypt(Source.Num(), Source.GetData(), Dest.GetData(), (RSA*)Key, RSA_PKCS1_PADDING);
+	if (NumEncryptedBytes == -1)
+	{
+		Dest.Empty(0);
+	}
+	return NumEncryptedBytes;
+}
+
+int32 FEncryptionContextOpenSSL::EncryptPrivate_RSA(TArrayView<const uint8> Source, TArray<uint8>& Dest, TPlatformCryptoRSAKey Key)
+{
+	Dest.SetNum(GetKeySize_RSA(Key));
+	int NumEncryptedBytes = RSA_private_encrypt(Source.Num(), Source.GetData(), Dest.GetData(), (RSA*)Key, RSA_PKCS1_PADDING);
+	if (NumEncryptedBytes == -1)
+	{
+		Dest.Empty(0);
+	}
+	return NumEncryptedBytes;
+}
+
+int32 FEncryptionContextOpenSSL::DecryptPublic_RSA(TArrayView<const uint8> Source, TArray<uint8>& Dest, TPlatformCryptoRSAKey Key)
+{
+	Dest.SetNum(GetKeySize_RSA(Key) - RSA_PKCS1_PADDING_SIZE);
+	int NumDecryptedBytes = RSA_public_decrypt(Source.Num(), Source.GetData(), Dest.GetData(), (RSA*)Key, RSA_PKCS1_PADDING);
+	if (NumDecryptedBytes == -1)
+	{
+		Dest.Empty(0);
+	}
+	else
+	{
+		Dest.SetNum(NumDecryptedBytes);
+	}
+	return NumDecryptedBytes;
+}
+
+int32 FEncryptionContextOpenSSL::DecryptPrivate_RSA(TArrayView<const uint8> Source, TArray<uint8>& Dest, TPlatformCryptoRSAKey Key)
+{
+	Dest.SetNum(GetKeySize_RSA(Key) - RSA_PKCS1_PADDING_SIZE);
+	int NumDecryptedBytes = RSA_private_decrypt(Source.Num(), Source.GetData(), Dest.GetData(), (RSA*)Key, RSA_PKCS1_PADDING);
+	if (NumDecryptedBytes == -1)
+	{
+		Dest.Empty(0);
+	}
+	else
+	{
+		Dest.SetNum(NumDecryptedBytes);
+	}
+	return NumDecryptedBytes;
+}
\ No newline at end of file
diff --git a/Engine/Plugins/Experimental/PlatformCrypto/Source/PlatformCryptoOpenSSL/Public/EncryptionContextOpenSSL.h b/Engine/Plugins/Experimental/PlatformCrypto/Source/PlatformCryptoOpenSSL/Public/EncryptionContextOpenSSL.h
--- a/Engine/Plugins/Experimental/PlatformCrypto/Source/PlatformCryptoOpenSSL/Public/EncryptionContextOpenSSL.h
+++ b/Engine/Plugins/Experimental/PlatformCrypto/Source/PlatformCryptoOpenSSL/Public/EncryptionContextOpenSSL.h
@@ -19,6 +19,15 @@ public:
 	TArray<uint8> Decrypt_AES_256_ECB(const TArrayView<const uint8> Ciphertext, const TArrayView<const uint8> Key, EPlatformCryptoResult& OutResult);
 
 	bool DigestVerify_PS256(const TArrayView<const char> Message, const TArrayView<const uint8> Signature, const TArrayView<const uint8> PKCS1Key);
+
+	TPlatformCryptoRSAKey CreateKey_RSA(const TArrayView<const uint8> PublicExponent, const TArrayView<const uint8> PrivateExponent, const TArrayView<const uint8> Modulus);
+	void DestroyKey_RSA(TPlatformCryptoRSAKey Key);
+	int32 GetKeySize_RSA(TPlatformCryptoRSAKey Key);
+	int32 GetMaxDataSize_RSA(TPlatformCryptoRSAKey Key);
+	int32 EncryptPublic_RSA(TArrayView<const uint8> Source, TArray<uint8>& Dest, TPlatformCryptoRSAKey Key);
+	int32 EncryptPrivate_RSA(TArrayView<const uint8> Source, TArray<uint8>& Dest, TPlatformCryptoRSAKey Key);
+	int32 DecryptPublic_RSA(TArrayView<const uint8> Source, TArray<uint8>& Dest, TPlatformCryptoRSAKey Key);
+	int32 DecryptPrivate_RSA(TArrayView<const uint8> Source, TArray<uint8>& Dest, TPlatformCryptoRSAKey Key);
 };
 
 typedef FEncryptionContextOpenSSL FEncryptionContext;
