Updated PlatformCrypto plugin in engine from EOSSDK Main (original CL 16755280)

diff --git a/Engine/Plugins/Experimental/PlatformCrypto/PlatformCrypto.uplugin b/Engine/Plugins/Experimental/PlatformCrypto/PlatformCrypto.uplugin
--- a/Engine/Plugins/Experimental/PlatformCrypto/PlatformCrypto.uplugin
+++ b/Engine/Plugins/Experimental/PlatformCrypto/PlatformCrypto.uplugin
@@ -26,7 +26,7 @@
 				"Mac",
 				"Win64",
 				"Linux",
-				"LinuxArm64",
+				"LinuxArm64"
 			],
 			"PlatformDenyList": [
 				"TVOS"
@@ -43,7 +43,7 @@
 				"Mac",
 				"Win64",
 				"Linux",
-				"LinuxArm64",
+				"LinuxArm64"
 			],
 			"PlatformDenyList": [
 				"TVOS"
@@ -59,11 +59,11 @@
 				"Mac",
 				"Win64",
 				"Linux",
-				"LinuxArm64",
+				"LinuxArm64"
 			],
 			"PlatformDenyList": [
 				"TVOS"
 			]
 		}
 	]
-}
+}
\ No newline at end of file
diff --git a/Engine/Plugins/Experimental/PlatformCrypto/Source/PlatformCrypto/Private/PlatformCrypto.cpp b/Engine/Plugins/Experimental/PlatformCrypto/Source/PlatformCrypto/Private/PlatformCrypto.cpp
--- a/Engine/Plugins/Experimental/PlatformCrypto/Source/PlatformCrypto/Private/PlatformCrypto.cpp
+++ b/Engine/Plugins/Experimental/PlatformCrypto/Source/PlatformCrypto/Private/PlatformCrypto.cpp
@@ -7,7 +7,7 @@
 #include "Features/IModularFeatures.h"
 #include "Misc/IEngineCrypto.h"
 
-class FPlatformCryptoModularFeature final : public IEngineCrypto
+class FPlatformCryptoModularFeature : public IEngineCrypto
 {
 public:
 
@@ -16,7 +16,7 @@ public:
 		IModularFeatures::Get().RegisterModularFeature(IEngineCrypto::GetFeatureName(), this);
 	}
 
-	~FPlatformCryptoModularFeature()
+	virtual ~FPlatformCryptoModularFeature()
 	{
 		Shutdown();
 	}
diff --git a/Engine/Plugins/Experimental/PlatformCrypto/Source/PlatformCryptoOpenSSL/Private/EncryptionContextOpenSSL.cpp b/Engine/Plugins/Experimental/PlatformCrypto/Source/PlatformCryptoOpenSSL/Private/EncryptionContextOpenSSL.cpp
--- a/Engine/Plugins/Experimental/PlatformCrypto/Source/PlatformCryptoOpenSSL/Private/EncryptionContextOpenSSL.cpp
+++ b/Engine/Plugins/Experimental/PlatformCrypto/Source/PlatformCryptoOpenSSL/Private/EncryptionContextOpenSSL.cpp
@@ -5,8 +5,70 @@
 #include "PlatformCryptoAesEncryptorsOpenSSL.h"
 #include "PlatformCryptoAesDecryptorsOpenSSL.h"
 
+#include "Misc/AssertionMacros.h"
+
 DEFINE_LOG_CATEGORY(LogPlatformCryptoOpenSSL);
 
+struct FSHA256HasherOpenSSL::FImplDetails
+{
+	SHA256_CTX Ctx;
+	enum class EState { NeedsInit, Ready, ReadyNeedsFinalize } State = EState::NeedsInit;
+};
+
+FSHA256HasherOpenSSL::FSHA256HasherOpenSSL() : Inner(MakePimpl<FSHA256HasherOpenSSL::FImplDetails>())
+{
+	check(Inner);
+	Init();
+}
+
+EPlatformCryptoResult FSHA256HasherOpenSSL::Init()
+{
+	if (Inner->State == FImplDetails::EState::ReadyNeedsFinalize)
+	{
+		UE_LOG(LogPlatformCryptoOpenSSL, Warning, TEXT("FSHA256Hasher::Init was called while there was still a digest being computed."));
+		return EPlatformCryptoResult::Failure;
+	}
+	if (Inner->State == FImplDetails::EState::Ready)
+	{
+		return EPlatformCryptoResult::Success;
+	}
+
+	SHA256_Init(&Inner->Ctx);
+	Inner->State = FImplDetails::EState::Ready;
+	return EPlatformCryptoResult::Success;
+}
+
+EPlatformCryptoResult FSHA256HasherOpenSSL::Update(const TArrayView<const uint8> InDataBuffer)
+{
+	if (Inner->State == FImplDetails::EState::NeedsInit)
+	{
+		UE_LOG(LogPlatformCryptoOpenSSL, Warning, TEXT("FSHA256Hasher::Update was called after the hasher was finalized without first calling Init"));
+		return EPlatformCryptoResult::Failure;
+	}
+
+	SHA256_Update(&Inner->Ctx, InDataBuffer.GetData(), static_cast<size_t>(InDataBuffer.Num()));
+	Inner->State = FImplDetails::EState::ReadyNeedsFinalize;
+	return EPlatformCryptoResult::Success;
+}
+
+EPlatformCryptoResult FSHA256HasherOpenSSL::Finalize(const TArrayView<uint8> OutDataBuffer)
+{
+	if (Inner->State == FImplDetails::EState::NeedsInit)
+	{
+		UE_LOG(LogPlatformCryptoOpenSSL, Warning, TEXT("FSHA256Hasher::Finalize was called after the hasher was finalized without first calling Init"));
+		return EPlatformCryptoResult::Failure;
+	}
+	if (OutDataBuffer.Num() < (int32)OutputByteLength)
+	{
+		UE_LOG(LogPlatformCryptoOpenSSL, Warning, TEXT("FSHA256Hasher::Finalize called with a buffer that is too small"));
+		return EPlatformCryptoResult::Failure;
+	}
+
+	SHA256_Final(reinterpret_cast<unsigned char*>(OutDataBuffer.GetData()), &Inner->Ctx);
+	Inner->State = FImplDetails::EState::NeedsInit;
+	return EPlatformCryptoResult::Success;
+}
+
 TArray<uint8> FEncryptionContextOpenSSL::Encrypt_AES_256_ECB(const TArrayView<const uint8> Plaintext, const TArrayView<const uint8> Key, EPlatformCryptoResult& OutResult)
 {
 	DECLARE_SCOPE_CYCLE_COUNTER(TEXT("OpenSSL AES256 ECB Encrypt"), STAT_OpenSSL_AES256_ECB_Encrypt, STATGROUP_PlatformCrypto);
@@ -75,14 +137,14 @@ TArray<uint8> FEncryptionContextOpenSSL::Encrypt_AES_256_CBC(const TArrayView<co
 	return Ciphertext;
 }
 
-TArray<uint8> FEncryptionContextOpenSSL::Encrypt_AES_256_GCM(const TArrayView<const uint8> Plaintext, const TArrayView<const uint8> Key, const TArrayView<const uint8> InitializationVector, TArray<uint8>& OutAuthTag, EPlatformCryptoResult& OutResult)
+TArray<uint8> FEncryptionContextOpenSSL::Encrypt_AES_256_GCM(const TArrayView<const uint8> Plaintext, const TArrayView<const uint8> Key, const TArrayView<const uint8> Nonce, TArray<uint8>& OutAuthTag, EPlatformCryptoResult& OutResult)
 {
 	DECLARE_SCOPE_CYCLE_COUNTER(TEXT("OpenSSL AES256 GCM Encrypt"), STAT_OpenSSL_AES256_GCM_Encrypt, STATGROUP_PlatformCrypto);
 
 	OutAuthTag.Reset();
 	OutResult = EPlatformCryptoResult::Failure;
 
-	TUniquePtr<IPlatformCryptoEncryptor> Encryptor = CreateEncryptor_AES_256_GCM(Key, InitializationVector);
+	TUniquePtr<IPlatformCryptoEncryptor> Encryptor = CreateEncryptor_AES_256_GCM(Key, Nonce);
 	if (!Encryptor.IsValid())
 	{
 		return TArray<uint8>();
@@ -189,13 +251,13 @@ TArray<uint8> FEncryptionContextOpenSSL::Decrypt_AES_256_CBC(const TArrayView<co
 	return Plaintext;
 }
 
-TArray<uint8> FEncryptionContextOpenSSL::Decrypt_AES_256_GCM(const TArrayView<const uint8> Ciphertext, const TArrayView<const uint8> Key, const TArrayView<const uint8> InitializationVector, const TArrayView<const uint8> AuthTag, EPlatformCryptoResult& OutResult)
+TArray<uint8> FEncryptionContextOpenSSL::Decrypt_AES_256_GCM(const TArrayView<const uint8> Ciphertext, const TArrayView<const uint8> Key, const TArrayView<const uint8> Nonce, const TArrayView<const uint8> AuthTag, EPlatformCryptoResult& OutResult)
 {
 	DECLARE_SCOPE_CYCLE_COUNTER(TEXT("OpenSSL AES256 GCM Decrypt"), STAT_OpenSSL_AES256_GCM_Decrypt, STATGROUP_PlatformCrypto);
 
 	OutResult = EPlatformCryptoResult::Failure;
 
-	TUniquePtr<IPlatformCryptoDecryptor> Decryptor = CreateDecryptor_AES_256_GCM(Key, InitializationVector, AuthTag);
+	TUniquePtr<IPlatformCryptoDecryptor> Decryptor = CreateDecryptor_AES_256_GCM(Key, Nonce, AuthTag);
 	if (!Decryptor.IsValid())
 	{
 		return TArray<uint8>();
@@ -233,9 +295,9 @@ TUniquePtr<IPlatformCryptoEncryptor> FEncryptionContextOpenSSL::CreateEncryptor_
 	return FPlatformCryptoEncryptor_AES_256_CBC_OpenSSL::Create(Key, InitializationVector);
 }
 
-TUniquePtr<IPlatformCryptoEncryptor> FEncryptionContextOpenSSL::CreateEncryptor_AES_256_GCM(const TArrayView<const uint8> Key, const TArrayView<const uint8> InitializationVector)
+TUniquePtr<IPlatformCryptoEncryptor> FEncryptionContextOpenSSL::CreateEncryptor_AES_256_GCM(const TArrayView<const uint8> Key, const TArrayView<const uint8> Nonce)
 {
-	return FPlatformCryptoEncryptor_AES_256_GCM_OpenSSL::Create(Key, InitializationVector);
+	return FPlatformCryptoEncryptor_AES_256_GCM_OpenSSL::Create(Key, Nonce);
 }
 
 TUniquePtr<IPlatformCryptoDecryptor> FEncryptionContextOpenSSL::CreateDecryptor_AES_256_ECB(const TArrayView<const uint8> Key)
@@ -248,9 +310,17 @@ TUniquePtr<IPlatformCryptoDecryptor> FEncryptionContextOpenSSL::CreateDecryptor_
 	return FPlatformCryptoDecryptor_AES_256_CBC_OpenSSL::Create(Key, InitializationVector);
 }
 
-TUniquePtr<IPlatformCryptoDecryptor> FEncryptionContextOpenSSL::CreateDecryptor_AES_256_GCM(const TArrayView<const uint8> Key, const TArrayView<const uint8> InitializationVector, const TArrayView<const uint8> AuthTag)
+TUniquePtr<IPlatformCryptoDecryptor> FEncryptionContextOpenSSL::CreateDecryptor_AES_256_GCM(const TArrayView<const uint8> Key, const TArrayView<const uint8> Nonce, const TArrayView<const uint8> AuthTag)
 {
-	return FPlatformCryptoDecryptor_AES_256_GCM_OpenSSL::Create(Key, InitializationVector, AuthTag);
+	return FPlatformCryptoDecryptor_AES_256_GCM_OpenSSL::Create(Key, Nonce, AuthTag);
+}
+
+bool FEncryptionContextOpenSSL::DigestSign_RS256(const TArrayView<const uint8> Message, TArray<uint8>& Signature, FRSAKeyHandle Key)
+{
+	Signature.SetNum(GetKeySize_RSA(Key));
+
+	unsigned int BytesWritten = 0;
+	return RSA_sign(NID_sha256, Message.GetData(), Message.Num(), Signature.GetData(), &BytesWritten, (RSA*)Key) == 1;
 }
 
 bool FEncryptionContextOpenSSL::DigestVerify_PS256(const TArrayView<const char> Message, const TArrayView<const uint8> Signature, const TArrayView<const uint8> PKCS1Key)
@@ -274,6 +344,11 @@ bool FEncryptionContextOpenSSL::DigestVerify_PS256(const TArrayView<const char>
 	return EVP_DigestVerifyFinal(Context.Get(), const_cast<uint8*>(Signature.GetData()), Signature.Num()) == 1;
 }
 
+bool FEncryptionContextOpenSSL::DigestVerify_RS256(const TArrayView<const uint8> Message, const TArrayView<const uint8> Signature, FRSAKeyHandle Key)
+{
+	return RSA_verify(NID_sha256, Message.GetData(), Message.Num(), Signature.GetData(), Signature.Num(), (RSA*)Key) == 1;
+}
+
 // Some platforms were upgraded to OpenSSL 1.1.1 while the others were left on a previous version. There are some minor differences we have to account for
 // in the older version, so declare a handy define that we can use to gate the code
 #if !defined(OPENSSL_VERSION_NUMBER) || OPENSSL_VERSION_NUMBER < 0x10100000L
@@ -282,6 +357,51 @@ bool FEncryptionContextOpenSSL::DigestVerify_PS256(const TArrayView<const char>
 #define USE_LEGACY_OPENSSL 0
 #endif
 
+void BigNumToArray(const int32 InKeySize, const BIGNUM* InNum, TArray<uint8>& OutBytes)
+{
+	int32 NumBytes = BN_num_bytes(InNum);
+	check(NumBytes <= InKeySize);
+	OutBytes.SetNumZeroed(NumBytes);
+
+	BN_bn2bin(InNum, OutBytes.GetData());
+	Algo::Reverse(OutBytes);
+}
+
+bool FEncryptionContextOpenSSL::GenerateKey_RSA(const int32 InNumKeyBits, TArray<uint8>& OutPublicExponent, TArray<uint8>& OutPrivateExponent, TArray<uint8>& OutModulus)
+{
+	if (InNumKeyBits % 8 != 0)
+	{
+		UE_LOG(LogPlatformCryptoOpenSSL, Warning, TEXT("FEncryptionContextOpenSSL::GenerateKey_RSA failed because InNumKeyBits is not divisible by 8."));
+		return false;
+	}
+
+	int32 KeySize = InNumKeyBits;
+	int32 KeySizeInBytes = InNumKeyBits / 8;
+
+	RSA* RSAKey = RSA_new();
+	BIGNUM* E = BN_new();
+	BN_set_word(E, RSA_F4);
+	RSA_generate_key_ex(RSAKey, KeySize, E, nullptr);
+
+#if USE_LEGACY_OPENSSL
+	const BIGNUM* PublicModulus = RSAKey->n;
+	const BIGNUM* PublicExponent = RSAKey->e;
+	const BIGNUM* PrivateExponent = RSAKey->d;
+#else
+	const BIGNUM* PublicModulus = RSA_get0_n(RSAKey);
+	const BIGNUM* PublicExponent = RSA_get0_e(RSAKey);
+	const BIGNUM* PrivateExponent = RSA_get0_d(RSAKey);
+#endif
+
+	BigNumToArray(KeySizeInBytes, PublicModulus, OutModulus);
+	BigNumToArray(KeySizeInBytes, PublicExponent, OutPublicExponent);
+	BigNumToArray(KeySizeInBytes, PrivateExponent, OutPrivateExponent);
+
+	RSA_free(RSAKey);
+
+	return true;
+}
+
 static void LoadBinaryIntoBigNum(const uint8* InData, int64 InDataSize, BIGNUM* InBigNum)
 {
 #if USE_LEGACY_OPENSSL
@@ -403,3 +523,8 @@ EPlatformCryptoResult FEncryptionContextOpenSSL::CreatePseudoRandomBytes(const T
 	return CreateRandomBytes(OutData);
 #endif // OPENSSL_VERSION_NUMBER < 0x10100000L
 }
+
+FSHA256Hasher FEncryptionContextOpenSSL::CreateSHA256Hasher()
+{
+	return FSHA256Hasher();
+}
diff --git a/Engine/Plugins/Experimental/PlatformCrypto/Source/PlatformCryptoOpenSSL/Private/PlatformCryptoAesDecryptorsOpenSSL.cpp b/Engine/Plugins/Experimental/PlatformCrypto/Source/PlatformCryptoOpenSSL/Private/PlatformCryptoAesDecryptorsOpenSSL.cpp
--- a/Engine/Plugins/Experimental/PlatformCrypto/Source/PlatformCryptoOpenSSL/Private/PlatformCryptoAesDecryptorsOpenSSL.cpp
+++ b/Engine/Plugins/Experimental/PlatformCrypto/Source/PlatformCryptoOpenSSL/Private/PlatformCryptoAesDecryptorsOpenSSL.cpp
@@ -81,9 +81,9 @@ EPlatformCryptoResult FPlatformCryptoDecryptor_AES_Base_OpenSSL::SetAuthTag(cons
 	return EPlatformCryptoResult::Success;
 }
 
-EPlatformCryptoResult FPlatformCryptoDecryptor_AES_Base_OpenSSL::Update(const TArrayView<const uint8> Plaintext, const TArrayView<uint8> OutCiphertext, int32& OutCiphertextBytesWritten)
+EPlatformCryptoResult FPlatformCryptoDecryptor_AES_Base_OpenSSL::Update(const TArrayView<const uint8> Ciphertext, const TArrayView<uint8> OutPlaintext, int32& OutPlaintextBytesWritten)
 {
-	OutCiphertextBytesWritten = 0;
+	OutPlaintextBytesWritten = 0;
 
 	if (State != EDecryptorState::Initialized)
 	{
@@ -91,14 +91,14 @@ EPlatformCryptoResult FPlatformCryptoDecryptor_AES_Base_OpenSSL::Update(const TA
 		return EPlatformCryptoResult::Failure;
 	}
 
-	const int32 RequiredBufferSize = GetUpdateBufferSizeBytes(Plaintext);
-	if (OutCiphertext.Num() < RequiredBufferSize)
+	const int32 RequiredBufferSize = GetUpdateBufferSizeBytes(Ciphertext);
+	if (OutPlaintext.Num() < RequiredBufferSize)
 	{
-		UE_LOG(LogPlatformCryptoOpenSSL, Warning, TEXT("FPlatformCryptoDecryptor_AES_Base_OpenSSL::Update: Invalid buffer size. Was %d needed %d"), OutCiphertext.Num() ,RequiredBufferSize);
+		UE_LOG(LogPlatformCryptoOpenSSL, Warning, TEXT("FPlatformCryptoDecryptor_AES_Base_OpenSSL::Update: Invalid buffer size. Was %d needed %d"), OutPlaintext.Num(), RequiredBufferSize);
 		return EPlatformCryptoResult::Failure;
 	}
 
-	const int UpdateResult = EVP_DecryptUpdate(EVPContext.Get(), OutCiphertext.GetData(), &OutCiphertextBytesWritten, Plaintext.GetData(), Plaintext.Num());
+	const int UpdateResult = EVP_DecryptUpdate(EVPContext.Get(), OutPlaintext.GetData(), &OutPlaintextBytesWritten, Ciphertext.GetData(), Ciphertext.Num());
 	if (UpdateResult != OPENSSL_CIPHER_SUCCESS)
 	{
 		UE_LOG(LogPlatformCryptoOpenSSL, Warning, TEXT("FPlatformCryptoDecryptor_AES_Base_OpenSSL::Update: EVP_DecryptInit_ex failed. Result=[%d]"), UpdateResult);
@@ -108,9 +108,9 @@ EPlatformCryptoResult FPlatformCryptoDecryptor_AES_Base_OpenSSL::Update(const TA
 	return EPlatformCryptoResult::Success;
 }
 
-EPlatformCryptoResult FPlatformCryptoDecryptor_AES_Base_OpenSSL::Finalize(const TArrayView<uint8> OutCiphertext, int32& OutCiphertextBytesWritten)
+EPlatformCryptoResult FPlatformCryptoDecryptor_AES_Base_OpenSSL::Finalize(const TArrayView<uint8> OutPlaintext, int32& OutPlaintextBytesWritten)
 {
-	OutCiphertextBytesWritten = 0;
+	OutPlaintextBytesWritten = 0;
 
 	if (State != EDecryptorState::Initialized)
 	{
@@ -118,12 +118,12 @@ EPlatformCryptoResult FPlatformCryptoDecryptor_AES_Base_OpenSSL::Finalize(const
 		return EPlatformCryptoResult::Failure;
 	}
 
-	if (OutCiphertext.Num() < GetFinalizeBufferSizeBytes())
+	if (OutPlaintext.Num() < GetFinalizeBufferSizeBytes())
 	{
 		return EPlatformCryptoResult::Failure;
 	}
 
-	const int FinalizeResult = EVP_DecryptFinal_ex(EVPContext.Get(), OutCiphertext.GetData(), &OutCiphertextBytesWritten);
+	const int FinalizeResult = EVP_DecryptFinal_ex(EVPContext.Get(), OutPlaintext.GetData(), &OutPlaintextBytesWritten);
 	if (FinalizeResult != OPENSSL_CIPHER_SUCCESS)
 	{
 		UE_LOG(LogPlatformCryptoOpenSSL, Warning, TEXT("FPlatformCryptoDecryptor_AES_Base_OpenSSL::Finalize: EVP_DecryptFinal_ex failed. Result=[%d]"), FinalizeResult);
@@ -153,7 +153,7 @@ TUniquePtr<IPlatformCryptoDecryptor> FPlatformCryptoDecryptor_AES_256_ECB_OpenSS
 	const EVP_CIPHER* Cipher = EVP_aes_256_ecb();
 
 	const int32 ExpectedKeyLength = EVP_CIPHER_key_length(Cipher);
-	if (Key.Num() < ExpectedKeyLength)
+	if (Key.Num() != ExpectedKeyLength)
 	{
 		UE_LOG(LogPlatformCryptoOpenSSL, Warning, TEXT("Invalid Key Size, failed to create Decryptor. KeySize=[%d] Expected=[%d]"), Key.Num(), ExpectedKeyLength);
 		return nullptr;
@@ -195,7 +195,7 @@ TUniquePtr<IPlatformCryptoDecryptor> FPlatformCryptoDecryptor_AES_256_CBC_OpenSS
 	const EVP_CIPHER* Cipher = EVP_aes_256_cbc();
 
 	const int32 ExpectedKeyLength = EVP_CIPHER_key_length(Cipher);
-	if (Key.Num() < ExpectedKeyLength)
+	if (Key.Num() != ExpectedKeyLength)
 	{
 		UE_LOG(LogPlatformCryptoOpenSSL, Warning, TEXT("Invalid Key Size, failed to create Decryptor. KeySize=[%d] Expected=[%d]"), Key.Num(), ExpectedKeyLength);
 		return nullptr;
@@ -244,7 +244,7 @@ TUniquePtr<IPlatformCryptoDecryptor> FPlatformCryptoDecryptor_AES_256_GCM_OpenSS
 	const EVP_CIPHER* Cipher = EVP_aes_256_gcm();
 
 	const int32 ExpectedKeyLength = EVP_CIPHER_key_length(Cipher);
-	if (Key.Num() < ExpectedKeyLength)
+	if (Key.Num() != ExpectedKeyLength)
 	{
 		UE_LOG(LogPlatformCryptoOpenSSL, Warning, TEXT("Invalid Key Size, failed to create Decryptor. KeySize=[%d] Expected=[%d]"), Key.Num(), ExpectedKeyLength);
 		return nullptr;
diff --git a/Engine/Plugins/Experimental/PlatformCrypto/Source/PlatformCryptoOpenSSL/Private/PlatformCryptoAesDecryptorsOpenSSL.h b/Engine/Plugins/Experimental/PlatformCrypto/Source/PlatformCryptoOpenSSL/Private/PlatformCryptoAesDecryptorsOpenSSL.h
--- a/Engine/Plugins/Experimental/PlatformCrypto/Source/PlatformCryptoOpenSSL/Private/PlatformCryptoAesDecryptorsOpenSSL.h
+++ b/Engine/Plugins/Experimental/PlatformCrypto/Source/PlatformCryptoOpenSSL/Private/PlatformCryptoAesDecryptorsOpenSSL.h
@@ -35,8 +35,8 @@ public:
 	virtual EPlatformCryptoResult SetAuthTag(const TArrayView<const uint8> AuthTag) override final;
 	virtual int32 GetCipherBlockSizeBytes() const override final;
 	virtual int32 GetCipherInitializationVectorSizeBytes() const override final;
-	virtual EPlatformCryptoResult Update(const TArrayView<const uint8> Plaintext, const TArrayView<uint8> OutCiphertext, int32& OutCiphertextBytesWritten) override final;
-	virtual EPlatformCryptoResult Finalize(const TArrayView<uint8> OutCiphertext, int32& OutCiphertextBytesWritten) override final;
+	virtual EPlatformCryptoResult Update(const TArrayView<const uint8> Ciphertext, const TArrayView<uint8> OutPlaintext, int32& OutPlaintextBytesWritten) override final;
+	virtual EPlatformCryptoResult Finalize(const TArrayView<uint8> OutPlaintext, int32& OutPlaintextBytesWritten) override final;
 	//~ End IPlatformCryptoDecryptor Interface
 
 	/**
diff --git a/Engine/Plugins/Experimental/PlatformCrypto/Source/PlatformCryptoOpenSSL/Private/PlatformCryptoAesEncryptorsOpenSSL.cpp b/Engine/Plugins/Experimental/PlatformCrypto/Source/PlatformCryptoOpenSSL/Private/PlatformCryptoAesEncryptorsOpenSSL.cpp
--- a/Engine/Plugins/Experimental/PlatformCrypto/Source/PlatformCryptoOpenSSL/Private/PlatformCryptoAesEncryptorsOpenSSL.cpp
+++ b/Engine/Plugins/Experimental/PlatformCrypto/Source/PlatformCryptoOpenSSL/Private/PlatformCryptoAesEncryptorsOpenSSL.cpp
@@ -160,7 +160,7 @@ TUniquePtr<IPlatformCryptoEncryptor> FPlatformCryptoEncryptor_AES_256_ECB_OpenSS
 	const EVP_CIPHER* Cipher = EVP_aes_256_ecb();
 
 	const int32 ExpectedKeyLength = EVP_CIPHER_key_length(Cipher);
-	if (Key.Num() < ExpectedKeyLength)
+	if (Key.Num() != ExpectedKeyLength)
 	{
 		UE_LOG(LogPlatformCryptoOpenSSL, Warning, TEXT("Invalid Key Size, failed to create Encryptor. KeySize=[%d] Expected=[%d]"), Key.Num(), ExpectedKeyLength);
 		return nullptr;
@@ -203,7 +203,7 @@ TUniquePtr<IPlatformCryptoEncryptor> FPlatformCryptoEncryptor_AES_256_CBC_OpenSS
 	const EVP_CIPHER* Cipher = EVP_aes_256_cbc();
 
 	const int32 ExpectedKeyLength = EVP_CIPHER_key_length(Cipher);
-	if (Key.Num() < ExpectedKeyLength)
+	if (Key.Num() != ExpectedKeyLength)
 	{
 		UE_LOG(LogPlatformCryptoOpenSSL, Warning, TEXT("Invalid Key Size, failed to create Encryptor. KeySize=[%d] Expected=[%d]"), Key.Num(), ExpectedKeyLength);
 		return nullptr;
@@ -252,7 +252,7 @@ TUniquePtr<IPlatformCryptoEncryptor> FPlatformCryptoEncryptor_AES_256_GCM_OpenSS
 	const EVP_CIPHER* Cipher = EVP_aes_256_gcm();
 
 	const int32 ExpectedKeyLength = EVP_CIPHER_key_length(Cipher);
-	if (Key.Num() < ExpectedKeyLength)
+	if (Key.Num() != ExpectedKeyLength)
 	{
 		UE_LOG(LogPlatformCryptoOpenSSL, Warning, TEXT("Invalid Key Size, failed to create Encryptor. KeySize=[%d] Expected=[%d]"), Key.Num(), ExpectedKeyLength);
 		return nullptr;
diff --git a/Engine/Plugins/Experimental/PlatformCrypto/Source/PlatformCryptoOpenSSL/Private/PlatformCryptoOpenSSLTypes.h b/Engine/Plugins/Experimental/PlatformCrypto/Source/PlatformCryptoOpenSSL/Private/PlatformCryptoOpenSSLTypes.h
--- a/Engine/Plugins/Experimental/PlatformCrypto/Source/PlatformCryptoOpenSSL/Private/PlatformCryptoOpenSSLTypes.h
+++ b/Engine/Plugins/Experimental/PlatformCrypto/Source/PlatformCryptoOpenSSL/Private/PlatformCryptoOpenSSLTypes.h
@@ -8,6 +8,7 @@ THIRD_PARTY_INCLUDES_START
 #include <openssl/evp.h>
 #include <openssl/rsa.h>
 #include <openssl/rand.h>
+#include <openssl/sha.h>
 THIRD_PARTY_INCLUDES_END
 
 DECLARE_LOG_CATEGORY_EXTERN(LogPlatformCryptoOpenSSL, Warning, All);
diff --git a/Engine/Plugins/Experimental/PlatformCrypto/Source/PlatformCryptoOpenSSL/Public/EncryptionContextOpenSSL.h b/Engine/Plugins/Experimental/PlatformCrypto/Source/PlatformCryptoOpenSSL/Public/EncryptionContextOpenSSL.h
--- a/Engine/Plugins/Experimental/PlatformCrypto/Source/PlatformCryptoOpenSSL/Public/EncryptionContextOpenSSL.h
+++ b/Engine/Plugins/Experimental/PlatformCrypto/Source/PlatformCryptoOpenSSL/Public/EncryptionContextOpenSSL.h
@@ -6,6 +6,50 @@
 #include "Containers/ArrayView.h"
 #include "PlatformCryptoTypes.h"
 
+#include "Templates/PimplPtr.h"
+
+/** Implementation details for SHA256 computation using OpenSSL */
+struct PLATFORMCRYPTOOPENSSL_API FSHA256HasherOpenSSL final
+{
+	FSHA256HasherOpenSSL(FSHA256HasherOpenSSL&&) = default;
+	FSHA256HasherOpenSSL& operator=(FSHA256HasherOpenSSL&&) = default;
+
+	/**
+	 * Initialize the necessary state to begin computing a message digest.
+	 * It is only necessary to call this function if the object is being used to compute multiple hashes, as the constructor will automatically call it the first time.
+	 */
+	EPlatformCryptoResult Init();
+
+	/**
+	 * Update the message digest computation with additional bytes
+	 *
+	 * @param InDataBuffer Buffer pointing to the data
+	 */
+	EPlatformCryptoResult Update(const TArrayView<const uint8> InDataBuffer);
+
+	/**
+	 * Finalize the computation of the message digest. After calling this method, Init must be called again if this object is meant to be reused for hashing a new input.
+	 *
+	 * @param OutDataBuffer A buffer that can hold the message digest bytes. Call GetOutputByteLength to determine the necessary size.
+	 */
+	EPlatformCryptoResult Finalize(const TArrayView<uint8> OutDataBuffer);
+
+	/**
+	 * The final message digest length in bytes.
+	 */
+	static constexpr const uint32 OutputByteLength = 32;
+
+private:
+	FSHA256HasherOpenSSL();
+	FSHA256HasherOpenSSL(const FSHA256HasherOpenSSL&) = delete;
+	FSHA256HasherOpenSSL& operator=(const FSHA256HasherOpenSSL&) = delete;
+	friend class FEncryptionContextOpenSSL;
+
+	struct FImplDetails;
+	TPimplPtr<FImplDetails> Inner;
+};
+using FSHA256Hasher = FSHA256HasherOpenSSL;
+
 /**
  * Interface to certain cryptographic algorithms, using OpenSSL to implement them.
  */
@@ -16,21 +60,26 @@ public:
 
 	TArray<uint8> Encrypt_AES_256_ECB(const TArrayView<const uint8> Plaintext, const TArrayView<const uint8> Key, EPlatformCryptoResult& OutResult);
 	TArray<uint8> Encrypt_AES_256_CBC(const TArrayView<const uint8> Plaintext, const TArrayView<const uint8> Key, const TArrayView<const uint8> InitializationVector, EPlatformCryptoResult& OutResult);
-	TArray<uint8> Encrypt_AES_256_GCM(const TArrayView<const uint8> Plaintext, const TArrayView<const uint8> Key, const TArrayView<const uint8> InitializationVector, TArray<uint8>& OutAuthTag, EPlatformCryptoResult& OutResult);
+	TArray<uint8> Encrypt_AES_256_GCM(const TArrayView<const uint8> Plaintext, const TArrayView<const uint8> Key, const TArrayView<const uint8> Nonce, TArray<uint8>& OutAuthTag, EPlatformCryptoResult& OutResult);
 
 	TArray<uint8> Decrypt_AES_256_ECB(const TArrayView<const uint8> Ciphertext, const TArrayView<const uint8> Key, EPlatformCryptoResult& OutResult);
 	TArray<uint8> Decrypt_AES_256_CBC(const TArrayView<const uint8> Ciphertext, const TArrayView<const uint8> Key, const TArrayView<const uint8> InitializationVector, EPlatformCryptoResult& OutResult);
-	TArray<uint8> Decrypt_AES_256_GCM(const TArrayView<const uint8> Ciphertext, const TArrayView<const uint8> Key, const TArrayView<const uint8> InitializationVector, const TArrayView<const uint8> AuthTag, EPlatformCryptoResult& OutResult);
+	TArray<uint8> Decrypt_AES_256_GCM(const TArrayView<const uint8> Ciphertext, const TArrayView<const uint8> Key, const TArrayView<const uint8> Nonce, const TArrayView<const uint8> AuthTag, EPlatformCryptoResult& OutResult);
 
 	TUniquePtr<IPlatformCryptoEncryptor> CreateEncryptor_AES_256_ECB(const TArrayView<const uint8> Key);
 	TUniquePtr<IPlatformCryptoEncryptor> CreateEncryptor_AES_256_CBC(const TArrayView<const uint8> Key, const TArrayView<const uint8> InitializationVector);
-	TUniquePtr<IPlatformCryptoEncryptor> CreateEncryptor_AES_256_GCM(const TArrayView<const uint8> Key, const TArrayView<const uint8> InitializationVector);
+	TUniquePtr<IPlatformCryptoEncryptor> CreateEncryptor_AES_256_GCM(const TArrayView<const uint8> Key, const TArrayView<const uint8> Nonce);
 
 	TUniquePtr<IPlatformCryptoDecryptor> CreateDecryptor_AES_256_ECB(const TArrayView<const uint8> Key);
 	TUniquePtr<IPlatformCryptoDecryptor> CreateDecryptor_AES_256_CBC(const TArrayView<const uint8> Key, const TArrayView<const uint8> InitializationVector);
-	TUniquePtr<IPlatformCryptoDecryptor> CreateDecryptor_AES_256_GCM(const TArrayView<const uint8> Key, const TArrayView<const uint8> InitializationVector, const TArrayView<const uint8> AuthTag);
+	TUniquePtr<IPlatformCryptoDecryptor> CreateDecryptor_AES_256_GCM(const TArrayView<const uint8> Key, const TArrayView<const uint8> Nonce, const TArrayView<const uint8> AuthTag);
+
+	bool DigestSign_RS256(const TArrayView<const uint8> Message, TArray<uint8>& Signature, FRSAKeyHandle Key);
 
 	bool DigestVerify_PS256(const TArrayView<const char> Message, const TArrayView<const uint8> Signature, const TArrayView<const uint8> PKCS1Key);
+	bool DigestVerify_RS256(const TArrayView<const uint8> Message, const TArrayView<const uint8> Signature, FRSAKeyHandle Key);
+
+	bool GenerateKey_RSA(const int32 InNumKeyBits, TArray<uint8>& OutPublicExponent, TArray<uint8>& OutPrivateExponent, TArray<uint8>& OutModulus);
 
 	FRSAKeyHandle CreateKey_RSA(const TArrayView<const uint8> PublicExponent, const TArrayView<const uint8> PrivateExponent, const TArrayView<const uint8> Modulus);
 	void DestroyKey_RSA(FRSAKeyHandle Key);
@@ -43,6 +92,8 @@ public:
 
 	EPlatformCryptoResult CreateRandomBytes(const TArrayView<uint8> OutData);
 	EPlatformCryptoResult CreatePseudoRandomBytes(const TArrayView<uint8> OutData);
+
+	FSHA256Hasher CreateSHA256Hasher();
 };
 
 typedef FEncryptionContextOpenSSL FEncryptionContext;
diff --git a/Engine/Plugins/Experimental/PlatformCrypto/Source/PlatformCryptoTypes/Private/PlatformCryptoTypes.cpp b/Engine/Plugins/Experimental/PlatformCrypto/Source/PlatformCryptoTypes/Private/PlatformCryptoTypes.cpp
--- a/Engine/Plugins/Experimental/PlatformCrypto/Source/PlatformCryptoTypes/Private/PlatformCryptoTypes.cpp
+++ b/Engine/Plugins/Experimental/PlatformCrypto/Source/PlatformCryptoTypes/Private/PlatformCryptoTypes.cpp
@@ -17,3 +17,113 @@ public:
 };
 
 IMPLEMENT_MODULE(FPlatformCryptoTypesModule, PlatformCryptoTypes)
+
+FAESBlockEncryptionHelper::FAESBlockEncryptionHelper(const int32 InBlockSize)
+	: BlockSize(InBlockSize)
+{
+	check(InBlockSize > 0);
+}
+
+EPlatformCryptoResult FAESBlockEncryptionHelper::Update(const TArrayView<const uint8> InDataBuffer, const TFunctionRef<FBlockHandlerFunctionSignature>& UpdateHandlingFunction, const TArrayView<uint8> OutDataBuffer, int32& OutBytesWritten)
+{
+	OutBytesWritten = 0;
+
+	// We need more than a block to be able to write out anything
+	if (LeftoverDataBuffer.Num() + InDataBuffer.Num() <= BlockSize)
+	{
+		// Don't have enough to write a block, just add it to our buffer for our next Update/Finalize
+		LeftoverDataBuffer.Append(InDataBuffer.GetData(), InDataBuffer.Num());
+		return EPlatformCryptoResult::Success;
+	}
+
+	// Make a mutable buffer objects that we will reassign as as finish processing data
+	TArrayView<const uint8> PlaintextLeftToProcess = InDataBuffer;
+	TArrayView<uint8> OutDataBufferAvailable = OutDataBuffer;
+
+	// Encrypt / Decrypt any previously-held chunk, potentially filling it in with new data
+	if (LeftoverDataBuffer.Num() > 0)
+	{
+		const int32 AmountToTake = BlockSize - LeftoverDataBuffer.Num();
+		check(AmountToTake >= 0);
+
+		if (AmountToTake > 0)
+		{
+			// Copy data into our buffer
+			LeftoverDataBuffer.Append(PlaintextLeftToProcess.GetData(), AmountToTake);
+
+			// Advance our data by how much we copied out
+			PlaintextLeftToProcess = TArrayView<const uint8>(PlaintextLeftToProcess.GetData() + AmountToTake, PlaintextLeftToProcess.Num() - AmountToTake);
+		}
+
+		int32 BytesWritten = 0;
+		const EPlatformCryptoResult UpdateResult = UpdateHandlingFunction(LeftoverDataBuffer, OutDataBufferAvailable, BytesWritten);
+		if (UpdateResult != EPlatformCryptoResult::Success)
+		{
+			return UpdateResult;
+		}
+
+		// Advance our output buffer and shrink size by how much we wrote
+		OutDataBufferAvailable = TArrayView<uint8>(OutDataBufferAvailable.GetData() + BytesWritten, OutDataBufferAvailable.Num() - BytesWritten);
+
+		// Write out how much data we've written
+		OutBytesWritten += BytesWritten;
+
+		// Reset our buffer (but do not free underlying memory)
+		LeftoverDataBuffer.Reset();
+	}
+
+	// Encrypt / Decrypt as many new chunks as we can, leaving one
+	while (PlaintextLeftToProcess.Num() > BlockSize)
+	{
+		check(OutDataBufferAvailable.Num() >= BlockSize);
+
+		int32 BytesWritten = 0;
+		const EPlatformCryptoResult UpdateResult = UpdateHandlingFunction(TArrayView<const uint8>(PlaintextLeftToProcess.GetData(), BlockSize), OutDataBufferAvailable, BytesWritten);
+		if (UpdateResult != EPlatformCryptoResult::Success)
+		{
+			return UpdateResult;
+		}
+
+		OutBytesWritten += BytesWritten;
+
+		// Advance our data by a block
+		PlaintextLeftToProcess = TArrayView<const uint8>(PlaintextLeftToProcess.GetData() + BytesWritten, PlaintextLeftToProcess.Num() - BytesWritten);
+
+		// Advance our output buffer by how much we wrote
+		OutDataBufferAvailable = TArrayView<uint8>(OutDataBufferAvailable.GetData() + BytesWritten, OutDataBufferAvailable.Num() - BytesWritten);
+	}
+
+	// Fill in any leftovers into our mismatched-size buffer
+	if (PlaintextLeftToProcess.Num() > 0)
+	{
+		LeftoverDataBuffer.Append(PlaintextLeftToProcess.GetData(), PlaintextLeftToProcess.Num());
+	}
+
+	return EPlatformCryptoResult::Success;
+}
+
+EPlatformCryptoResult FAESBlockEncryptionHelper::Finalize(const TFunctionRef<FBlockHandlerFunctionSignature>& FinalizeHandlingFunction, const TArrayView<uint8> OutDataBuffer, int32& OutBytesWritten)
+{
+	OutBytesWritten = 0;
+
+	// If we have no data, it means that Update was never called (which is fine)
+	// If update has never been called, we don't have to do anything
+	if (LeftoverDataBuffer.Num() > 0)
+	{
+		// Process our final chunk of data
+		int32 BytesWritten = 0;
+		EPlatformCryptoResult FinalResult = FinalizeHandlingFunction(LeftoverDataBuffer, OutDataBuffer, BytesWritten);
+		if (FinalResult != EPlatformCryptoResult::Success)
+		{
+			return FinalResult;
+		}
+
+		// Write out how much data we wrote
+		OutBytesWritten += BytesWritten;
+
+		// Reset our buffer (but do not free underlying memory)
+		LeftoverDataBuffer.Reset();
+	}
+
+	return EPlatformCryptoResult::Success;
+}
diff --git a/Engine/Plugins/Experimental/PlatformCrypto/Source/PlatformCryptoTypes/Public/PlatformCryptoTypes.h b/Engine/Plugins/Experimental/PlatformCrypto/Source/PlatformCryptoTypes/Public/PlatformCryptoTypes.h
--- a/Engine/Plugins/Experimental/PlatformCrypto/Source/PlatformCryptoTypes/Public/PlatformCryptoTypes.h
+++ b/Engine/Plugins/Experimental/PlatformCrypto/Source/PlatformCryptoTypes/Public/PlatformCryptoTypes.h
@@ -189,3 +189,42 @@ public:
 	 */
 	virtual EPlatformCryptoResult Finalize(const TArrayView<uint8> OutPlaintext, int32& OutPlaintextBytesWritten) = 0;
 };
+
+/**
+ * Helper class to ensure that only multiples of block-size get passed into our Encrypt/Decrypt functions.
+ * This class will hold onto the last-full or partial block until another call to Update or Finalize happens.
+ */
+class FAESBlockEncryptionHelper
+{
+public:
+	FAESBlockEncryptionHelper(const int32 BlockSize);
+
+	using FBlockHandlerFunctionSignature = EPlatformCryptoResult(const TArrayView<const uint8> /*InDataBuffer*/, const TArrayView<uint8> /*OutDataBuffer*/, int32& /*OutBytesWritten*/);
+
+	/**
+	 * Add text to potentialy be processed by the UpdateHandlingFunction. If there is not enough text, the UpdateHandlingFunction may not be called.
+	 *
+	 * @param InDataBuffer Buffer pointing to the data that needs to be encrypted/decrypted
+	 * @param UpdateHandlingFunction A function that can encrypt/decrypt a complete block of text
+	 * @param OutDataBuffer A buffer that can hold the complete encrypted/decrypted data from UpdateHandlingFunction
+	 * @param OutBytesWritten The amount of bytes written to OutDataBuffer
+	 * @return Success if the data could be successfully processed, or Failure otherwise.
+	 */
+	EPlatformCryptoResult Update(const TArrayView<const uint8> InDataBuffer, const TFunctionRef<FBlockHandlerFunctionSignature>& UpdateHandlingFunction, const TArrayView<uint8> OutDataBuffer, int32& OutBytesWritten);
+
+	/**
+	 * Process any remaining block of text that was previously sent to Update. If there are no blocks remaining, the FinalizeHandlingFunction may not be called.
+	 *
+	 * @param FinalizeHandlingFunction A function that can encrypt/decrypt a complete block of text
+	 * @param OutDataBuffer A buffer that can hold the complete encrypted/decrypted data from FinalizeHandlingFunction
+	 * @param OutBytesWritten The amount of bytes written to OutDataBuffer
+	 * @return Success if the data could be successfully processed, or Failure otherwise.
+	 */
+	EPlatformCryptoResult Finalize(const TFunctionRef<FBlockHandlerFunctionSignature>& FinalizeHandlingFunction, const TArrayView<uint8> OutDataBuffer, int32& OutBytesWritten);
+
+protected:
+	/** A buffer to store data waiting to be encrypted/decrypted */
+	TArray<uint8> LeftoverDataBuffer;
+	/** The size of a complete block of data */
+	int32 BlockSize = 0;
+};
