modular build fix for FAESBlockEncryptionHelper.

#rnx

diff --git a/Engine/Plugins/Experimental/PlatformCrypto/Source/PlatformCryptoTypes/Public/PlatformCryptoTypes.h b/Engine/Plugins/Experimental/PlatformCrypto/Source/PlatformCryptoTypes/Public/PlatformCryptoTypes.h
--- a/Engine/Plugins/Experimental/PlatformCrypto/Source/PlatformCryptoTypes/Public/PlatformCryptoTypes.h
+++ b/Engine/Plugins/Experimental/PlatformCrypto/Source/PlatformCryptoTypes/Public/PlatformCryptoTypes.h
@@ -194,7 +194,7 @@ public:
  * Helper class to ensure that only multiples of block-size get passed into our Encrypt/Decrypt functions.
  * This class will hold onto the last-full or partial block until another call to Update or Finalize happens.
  */
-class FAESBlockEncryptionHelper
+class PLATFORMCRYPTOTYPES_API FAESBlockEncryptionHelper
 {
 public:
 	FAESBlockEncryptionHelper(const int32 BlockSize);
