Fix for static analysis warning (call to Shutdown() virtual function in FPlatformCryptoModularFeature destructor).

diff --git a/Engine/Plugins/Experimental/PlatformCrypto/Source/PlatformCrypto/Private/PlatformCrypto.cpp b/Engine/Plugins/Experimental/PlatformCrypto/Source/PlatformCrypto/Private/PlatformCrypto.cpp
--- a/Engine/Plugins/Experimental/PlatformCrypto/Source/PlatformCrypto/Private/PlatformCrypto.cpp
+++ b/Engine/Plugins/Experimental/PlatformCrypto/Source/PlatformCrypto/Private/PlatformCrypto.cpp
@@ -18,12 +18,12 @@ public:
 
 	virtual ~FPlatformCryptoModularFeature()
 	{
-		Shutdown();
+		Cleanup();
 	}
 
 	virtual void Shutdown() override
 	{
-		Context.Reset();
+		Cleanup();
 	}
 	
 	/** IEngineCrypto implementation */
@@ -68,6 +68,10 @@ public:
 	}
 
 private:
+	void Cleanup()
+	{
+		Context.Reset();
+	}
 
 	TUniquePtr<FEncryptionContext>& GetContext()
 	{
