Remove console NDA references from PlatformCrypto

diff --git a/Engine/Plugins/Experimental/PlatformCrypto/Source/PlatformCrypto/PlatformCrypto.Build.cs b/Engine/Plugins/Experimental/PlatformCrypto/Source/PlatformCrypto/PlatformCrypto.Build.cs
--- a/Engine/Plugins/Experimental/PlatformCrypto/Source/PlatformCrypto/PlatformCrypto.Build.cs
+++ b/Engine/Plugins/Experimental/PlatformCrypto/Source/PlatformCrypto/PlatformCrypto.Build.cs
@@ -25,6 +25,8 @@ namespace UnrealBuildTool.Rules
 					}
 					);
 			}
+			PublicDefinitions.Add("PLATFORM_ENCRYPTION_CONTEXT_HEADER=\"" + PlatformEncryptionContextHeader + "\"");
 		}
+		protected virtual string PlatformEncryptionContextHeader { get { return "EncryptionContextOpenSSL.h"; } }
 	}
 }
diff --git a/Engine/Plugins/Experimental/PlatformCrypto/Source/PlatformCrypto/Public/PlatformCryptoIncludes.h b/Engine/Plugins/Experimental/PlatformCrypto/Source/PlatformCrypto/Public/PlatformCryptoIncludes.h
--- a/Engine/Plugins/Experimental/PlatformCrypto/Source/PlatformCrypto/Public/PlatformCryptoIncludes.h
+++ b/Engine/Plugins/Experimental/PlatformCrypto/Source/PlatformCrypto/Public/PlatformCryptoIncludes.h
@@ -4,14 +4,4 @@
 
 #include "CoreMinimal.h"
 
-#if !defined(USE_BCRYPT)
-	#define USE_BCRYPT 0
-#endif
-
-#if USE_BCRYPT
-	#include "EncryptionContextBCrypt.h"
-#elif PLATFORM_SWITCH
-	#include "EncryptionContextSwitch.h"
-#else
-	#include "EncryptionContextOpenSSL.h"
-#endif
\ No newline at end of file
+#include PLATFORM_ENCRYPTION_CONTEXT_HEADER
