Fixed IWYU markups

diff --git a/Engine/Plugins/Experimental/PlatformCrypto/Source/PlatformCrypto/Public/PlatformCryptoIncludes.h b/Engine/Plugins/Experimental/PlatformCrypto/Source/PlatformCrypto/Public/PlatformCryptoIncludes.h
--- a/Engine/Plugins/Experimental/PlatformCrypto/Source/PlatformCrypto/Public/PlatformCryptoIncludes.h
+++ b/Engine/Plugins/Experimental/PlatformCrypto/Source/PlatformCrypto/Public/PlatformCryptoIncludes.h
@@ -4,4 +4,6 @@
 
 #include "CoreMinimal.h"
 
+// IWYU pragma: begin_exports
 #include PLATFORM_ENCRYPTION_CONTEXT_HEADER
+// IWYU pragma: end_exports
\ No newline at end of file
diff --git a/Engine/Source/Runtime/Core/Public/Math/UnrealMathVectorConstants.h b/Engine/Source/Runtime/Core/Public/Math/UnrealMathVectorConstants.h
--- a/Engine/Source/Runtime/Core/Public/Math/UnrealMathVectorConstants.h
+++ b/Engine/Source/Runtime/Core/Public/Math/UnrealMathVectorConstants.h
@@ -3,6 +3,7 @@
 #pragma once
 
 // HEADER_UNIT_SKIP - Not included directly
+// IWYU pragma: private
 
 #include "Math/UnrealMathUtility.h"
 
