[PlatformCrypto]
* Changed include from Stats2.h back to Stats.h

diff --git a/Engine/Plugins/Experimental/PlatformCrypto/Source/PlatformCryptoTypes/Public/PlatformCryptoTypes.h b/Engine/Plugins/Experimental/PlatformCrypto/Source/PlatformCryptoTypes/Public/PlatformCryptoTypes.h
--- a/Engine/Plugins/Experimental/PlatformCrypto/Source/PlatformCryptoTypes/Public/PlatformCryptoTypes.h
+++ b/Engine/Plugins/Experimental/PlatformCrypto/Source/PlatformCryptoTypes/Public/PlatformCryptoTypes.h
@@ -2,7 +2,7 @@
 
 #pragma once
 
-#include "Stats/Stats2.h"
+#include "Stats/Stats.h" // IWYU pragma: keep
 
 /**
  * Generic result type for cryptographic functions.
@@ -230,5 +230,4 @@ protected:
 #if UE_ENABLE_INCLUDE_ORDER_DEPRECATED_IN_5_2
 #include "CoreMinimal.h"
 #include "Misc/IEngineCrypto.h"
-#include "Stats/Stats.h"
 #endif
