[PlatformCrypto] Adding GetPublicKey_RSA method to convert bytes to RSA key

Adding the GetPublicKey_RSA method to convert raw bytes (i.e. coming from a string) to an RSA key in OpenSSL representation.

diff --git a/Engine/Plugins/Experimental/PlatformCrypto/Source/PlatformCryptoOpenSSL/Private/EncryptionContextOpenSSL.cpp b/Engine/Plugins/Experimental/PlatformCrypto/Source/PlatformCryptoOpenSSL/Private/EncryptionContextOpenSSL.cpp
--- a/Engine/Plugins/Experimental/PlatformCrypto/Source/PlatformCryptoOpenSSL/Private/EncryptionContextOpenSSL.cpp
+++ b/Engine/Plugins/Experimental/PlatformCrypto/Source/PlatformCryptoOpenSSL/Private/EncryptionContextOpenSSL.cpp
@@ -10,6 +10,7 @@
 #include <openssl/rand.h>
 #include <openssl/rsa.h>
 #include <openssl/sha.h>
+#include <openssl/pem.h>
 
 DEFINE_LOG_CATEGORY(LogPlatformCryptoOpenSSL);
 
@@ -447,6 +448,16 @@ FRSAKeyHandle FEncryptionContextOpenSSL::CreateKey_RSA(const TArrayView<const ui
 	return NewKey;
 }
 
+FRSAKeyHandle FEncryptionContextOpenSSL::GetPublicKey_RSA(const TArrayView<const uint8> Source)
+{
+	BIO* KeyBuf = BIO_new_mem_buf(Source.GetData(), Source.Num());
+	RSA* Key = PEM_read_bio_RSA_PUBKEY(KeyBuf, 0, 0, 0);
+
+	BIO_free(KeyBuf);
+
+	return Key;
+}
+
 void FEncryptionContextOpenSSL::DestroyKey_RSA(FRSAKeyHandle Key)
 {
 	RSA_free((RSA*)Key);
diff --git a/Engine/Plugins/Experimental/PlatformCrypto/Source/PlatformCryptoOpenSSL/Public/EncryptionContextOpenSSL.h b/Engine/Plugins/Experimental/PlatformCrypto/Source/PlatformCryptoOpenSSL/Public/EncryptionContextOpenSSL.h
--- a/Engine/Plugins/Experimental/PlatformCrypto/Source/PlatformCryptoOpenSSL/Public/EncryptionContextOpenSSL.h
+++ b/Engine/Plugins/Experimental/PlatformCrypto/Source/PlatformCryptoOpenSSL/Public/EncryptionContextOpenSSL.h
@@ -84,6 +84,7 @@ public:
 	bool GenerateKey_RSA(const int32 InNumKeyBits, TArray<uint8>& OutPublicExponent, TArray<uint8>& OutPrivateExponent, TArray<uint8>& OutModulus);
 
 	FRSAKeyHandle CreateKey_RSA(const TArrayView<const uint8> PublicExponent, const TArrayView<const uint8> PrivateExponent, const TArrayView<const uint8> Modulus);
+	FRSAKeyHandle GetPublicKey_RSA(const TArrayView<const uint8> Source);
 	void DestroyKey_RSA(FRSAKeyHandle Key);
 	int32 GetKeySize_RSA(FRSAKeyHandle Key);
 	int32 GetMaxDataSize_RSA(FRSAKeyHandle Key);
