[PlatformCrypto] Adding utility method to calculate SHA256 hash

Adding a utility method to calculate SHA256 hashes. One of the applications will be to verify JWTs down the line.

diff --git a/Engine/Plugins/Experimental/PlatformCrypto/Source/PlatformCryptoOpenSSL/Private/EncryptionContextOpenSSL.cpp b/Engine/Plugins/Experimental/PlatformCrypto/Source/PlatformCryptoOpenSSL/Private/EncryptionContextOpenSSL.cpp
--- a/Engine/Plugins/Experimental/PlatformCrypto/Source/PlatformCryptoOpenSSL/Private/EncryptionContextOpenSSL.cpp
+++ b/Engine/Plugins/Experimental/PlatformCrypto/Source/PlatformCryptoOpenSSL/Private/EncryptionContextOpenSSL.cpp
@@ -543,3 +543,37 @@ FSHA256Hasher FEncryptionContextOpenSSL::CreateSHA256Hasher()
 {
 	return FSHA256Hasher();
 }
+
+bool FEncryptionContextOpenSSL::CalcSHA256(const TArrayView<const uint8> Source, TArray<uint8>& OutHash)
+{
+	FSHA256Hasher Hasher = CreateSHA256Hasher();
+
+	if (Hasher.Init() != EPlatformCryptoResult::Success)
+	{
+		UE_LOG(LogPlatformCryptoOpenSSL, Verbose,
+			TEXT("[FEncryptionContextOpenSSL::CalcSHA256] Failed to create SHA256Hasher."));
+
+		return false;
+	}
+
+	if (Hasher.Update(Source) != EPlatformCryptoResult::Success)
+	{
+		UE_LOG(LogPlatformCryptoOpenSSL, Verbose,
+			TEXT("[FEncryptionContextOpenSSL::CalcSHA256] Failed to update SHA256Hasher."));
+
+		return false;
+	}
+
+	OutHash.Empty();
+	OutHash.AddUninitialized(FSHA256Hasher::OutputByteLength);
+
+	if (Hasher.Finalize(OutHash) != EPlatformCryptoResult::Success)
+	{
+		UE_LOG(LogPlatformCryptoOpenSSL, Verbose,
+			TEXT("[FEncryptionContextOpenSSL::CalcSHA256] Failed to finalize SHA256Hasher."));
+
+		return false;
+	}
+
+	return true;
+}
diff --git a/Engine/Plugins/Experimental/PlatformCrypto/Source/PlatformCryptoOpenSSL/Public/EncryptionContextOpenSSL.h b/Engine/Plugins/Experimental/PlatformCrypto/Source/PlatformCryptoOpenSSL/Public/EncryptionContextOpenSSL.h
--- a/Engine/Plugins/Experimental/PlatformCrypto/Source/PlatformCryptoOpenSSL/Public/EncryptionContextOpenSSL.h
+++ b/Engine/Plugins/Experimental/PlatformCrypto/Source/PlatformCryptoOpenSSL/Public/EncryptionContextOpenSSL.h
@@ -97,6 +97,9 @@ public:
 	EPlatformCryptoResult CreatePseudoRandomBytes(const TArrayView<uint8> OutData);
 
 	FSHA256Hasher CreateSHA256Hasher();
+
+	bool CalcSHA256(const TArrayView<const uint8> Source, TArray<uint8>& OutHash);
+
 };
 
 typedef FEncryptionContextOpenSSL FEncryptionContext;
