[PlatformCrypto] Adding OpenSSL error reporting to DigestVerify_RS256

Adding OpenSSL error reporting to the DigestVerify_RS256 method.

diff --git a/Engine/Plugins/Experimental/PlatformCrypto/Source/PlatformCryptoOpenSSL/Private/EncryptionContextOpenSSL.cpp b/Engine/Plugins/Experimental/PlatformCrypto/Source/PlatformCryptoOpenSSL/Private/EncryptionContextOpenSSL.cpp
--- a/Engine/Plugins/Experimental/PlatformCrypto/Source/PlatformCryptoOpenSSL/Private/EncryptionContextOpenSSL.cpp
+++ b/Engine/Plugins/Experimental/PlatformCrypto/Source/PlatformCryptoOpenSSL/Private/EncryptionContextOpenSSL.cpp
@@ -4,6 +4,7 @@
 #include "PlatformCryptoAesEncryptorsOpenSSL.h"
 #include "PlatformCryptoAesDecryptorsOpenSSL.h"
 
+THIRD_PARTY_INCLUDES_START
 #include <openssl/bn.h>
 #include <openssl/obj_mac.h>
 #include <openssl/opensslv.h>
@@ -11,6 +12,9 @@
 #include <openssl/rsa.h>
 #include <openssl/sha.h>
 #include <openssl/pem.h>
+#include <openssl/err.h>
+THIRD_PARTY_INCLUDES_END
+
 
 DEFINE_LOG_CATEGORY(LogPlatformCryptoOpenSSL);
 
@@ -351,7 +355,16 @@ bool FEncryptionContextOpenSSL::DigestVerify_PS256(const TArrayView<const char>
 
 bool FEncryptionContextOpenSSL::DigestVerify_RS256(const TArrayView<const uint8> Message, const TArrayView<const uint8> Signature, FRSAKeyHandle Key)
 {
-	return RSA_verify(NID_sha256, Message.GetData(), Message.Num(), Signature.GetData(), Signature.Num(), (RSA*)Key) == 1;
+	const int VerifyResult = RSA_verify(NID_sha256, Message.GetData(), Message.Num(), Signature.GetData(), Signature.Num(), (RSA*)Key);
+
+	if (VerifyResult != 1)
+	{
+		UE_LOG(LogPlatformCryptoOpenSSL, Error,
+			TEXT("[FEncryptionContextOpenSSL::DigestVerify_RS256] Signature is invalid: (%u): %s."),
+				ERR_get_error(), UTF8_TO_TCHAR(ERR_error_string(ERR_get_error(), nullptr)));
+	}
+
+	return VerifyResult == 1;
 }
 
 // Some platforms were upgraded to OpenSSL 1.1.1 while the others were left on a previous version. There are some minor differences we have to account for
