[PlatformCrypto] Adding SupportedTargetPlatforms to PlatformCrypto

Adding SupportedTargetPlatforms to PlatformCrypto in order to hopefully get rid of the TVOS build errors.

diff --git a/Engine/Plugins/Experimental/PlatformCrypto/PlatformCrypto.uplugin b/Engine/Plugins/Experimental/PlatformCrypto/PlatformCrypto.uplugin
--- a/Engine/Plugins/Experimental/PlatformCrypto/PlatformCrypto.uplugin
+++ b/Engine/Plugins/Experimental/PlatformCrypto/PlatformCrypto.uplugin
@@ -15,6 +15,14 @@
 	"IsBetaVersion": false,
 	"Installed": false,
 	"SupportedPrograms": [ "UnrealPak", "UnrealPakEx" ],
+	"SupportedTargetPlatforms": [
+		"Android",
+		"IOS",
+		"Mac",
+		"Win64",
+		"Linux",
+		"LinuxArm64"
+	],
 	"Modules": [
 		{
 			"Name": "PlatformCrypto",
