[PlatformCrypto] Reverting PlatformCrypto SupportedTargetPlatforms change

diff --git a/Engine/Plugins/Experimental/PlatformCrypto/PlatformCrypto.uplugin b/Engine/Plugins/Experimental/PlatformCrypto/PlatformCrypto.uplugin
--- a/Engine/Plugins/Experimental/PlatformCrypto/PlatformCrypto.uplugin
+++ b/Engine/Plugins/Experimental/PlatformCrypto/PlatformCrypto.uplugin
@@ -15,14 +15,6 @@
 	"IsBetaVersion": false,
 	"Installed": false,
 	"SupportedPrograms": [ "UnrealPak", "UnrealPakEx" ],
-	"SupportedTargetPlatforms": [
-		"Android",
-		"IOS",
-		"Mac",
-		"Win64",
-		"Linux",
-		"LinuxArm64"
-	],
 	"Modules": [
 		{
 			"Name": "PlatformCrypto",
