[PlatformCrypto] Adding FStringView specialized GetPublicKey_RSA()

Adding a specialized method of GetPublicKey_RSA() to convert FStringView source RSA public keys into the OpenSSL representation.

diff --git a/Engine/Plugins/Experimental/PlatformCrypto/Source/PlatformCryptoOpenSSL/Private/EncryptionContextOpenSSL.cpp b/Engine/Plugins/Experimental/PlatformCrypto/Source/PlatformCryptoOpenSSL/Private/EncryptionContextOpenSSL.cpp
--- a/Engine/Plugins/Experimental/PlatformCrypto/Source/PlatformCryptoOpenSSL/Private/EncryptionContextOpenSSL.cpp
+++ b/Engine/Plugins/Experimental/PlatformCrypto/Source/PlatformCryptoOpenSSL/Private/EncryptionContextOpenSSL.cpp
@@ -4,6 +4,8 @@
 #include "PlatformCryptoAesEncryptorsOpenSSL.h"
 #include "PlatformCryptoAesDecryptorsOpenSSL.h"
 
+#include "Containers/StringConv.h"
+
 THIRD_PARTY_INCLUDES_START
 #include <openssl/bn.h>
 #include <openssl/obj_mac.h>
@@ -461,9 +463,9 @@ FRSAKeyHandle FEncryptionContextOpenSSL::CreateKey_RSA(const TArrayView<const ui
 	return NewKey;
 }
 
-FRSAKeyHandle FEncryptionContextOpenSSL::GetPublicKey_RSA(const TArrayView<const uint8> Source)
+FRSAKeyHandle FEncryptionContextOpenSSL::GetPublicKey_RSA(const void* Source, const int32 Length)
 {
-	BIO* KeyBuf = BIO_new_mem_buf(Source.GetData(), Source.Num());
+	BIO* KeyBuf = BIO_new_mem_buf(Source, Length);
 	RSA* Key = PEM_read_bio_RSA_PUBKEY(KeyBuf, 0, 0, 0);
 
 	BIO_free(KeyBuf);
@@ -471,6 +473,18 @@ FRSAKeyHandle FEncryptionContextOpenSSL::GetPublicKey_RSA(const TArrayView<const
 	return Key;
 }
 
+FRSAKeyHandle FEncryptionContextOpenSSL::GetPublicKey_RSA(const TArrayView<const uint8> Source)
+{
+	return GetPublicKey_RSA(Source.GetData(), Source.Num());
+}
+
+FRSAKeyHandle FEncryptionContextOpenSSL::GetPublicKey_RSA(const FStringView Source)
+{
+	auto ConvertedSource = StringCast<UTF8CHAR>(Source.GetData());
+
+	return GetPublicKey_RSA(ConvertedSource.Get(), ConvertedSource.Length());
+}
+
 void FEncryptionContextOpenSSL::DestroyKey_RSA(FRSAKeyHandle Key)
 {
 	RSA_free((RSA*)Key);
diff --git a/Engine/Plugins/Experimental/PlatformCrypto/Source/PlatformCryptoOpenSSL/Public/EncryptionContextOpenSSL.h b/Engine/Plugins/Experimental/PlatformCrypto/Source/PlatformCryptoOpenSSL/Public/EncryptionContextOpenSSL.h
--- a/Engine/Plugins/Experimental/PlatformCrypto/Source/PlatformCryptoOpenSSL/Public/EncryptionContextOpenSSL.h
+++ b/Engine/Plugins/Experimental/PlatformCrypto/Source/PlatformCryptoOpenSSL/Public/EncryptionContextOpenSSL.h
@@ -84,7 +84,9 @@ public:
 	bool GenerateKey_RSA(const int32 InNumKeyBits, TArray<uint8>& OutPublicExponent, TArray<uint8>& OutPrivateExponent, TArray<uint8>& OutModulus);
 
 	FRSAKeyHandle CreateKey_RSA(const TArrayView<const uint8> PublicExponent, const TArrayView<const uint8> PrivateExponent, const TArrayView<const uint8> Modulus);
+	FRSAKeyHandle GetPublicKey_RSA(const void* Source, const int32 Length);
 	FRSAKeyHandle GetPublicKey_RSA(const TArrayView<const uint8> Source);
+	FRSAKeyHandle GetPublicKey_RSA(const FStringView Source);
 	void DestroyKey_RSA(FRSAKeyHandle Key);
 	int32 GetKeySize_RSA(FRSAKeyHandle Key);
 	int32 GetMaxDataSize_RSA(FRSAKeyHandle Key);
