[PlatformCrypto] Fixing GetPublicKey_RSA() specializations and issues

Since GetPublicKey_RSA() is using a PEM_read function within OpenSSL, binary string representation specializations are not necessary. Additionally, key handling on the switch platform with GetPublicKey_RSA() was previously broken and now conforms to the key handling that is used in the rest of the platform specific plugin.

diff --git a/Engine/Plugins/Experimental/JWT/Source/JWT/Private/JwtAlgorithms.cpp b/Engine/Plugins/Experimental/JWT/Source/JWT/Private/JwtAlgorithms.cpp
--- a/Engine/Plugins/Experimental/JWT/Source/JWT/Private/JwtAlgorithms.cpp
+++ b/Engine/Plugins/Experimental/JWT/Source/JWT/Private/JwtAlgorithms.cpp
@@ -40,7 +40,7 @@ bool FJwtAlgorithm_RS256::VerifySignature(
 }
 
 
-bool FJwtAlgorithm_RS256::SetPublicKey(const TArrayView<const uint8> InKey)
+bool FJwtAlgorithm_RS256::SetPublicKey(const FStringView InPemKey)
 {
 	if (!EncryptionContext.IsValid())
 	{
@@ -52,12 +52,13 @@ bool FJwtAlgorithm_RS256::SetPublicKey(const TArrayView<const uint8> InKey)
 		DestroyKey(PublicKey);
 	}
 
-	PublicKey = EncryptionContext->GetPublicKey_RSA(InKey);
+	PublicKey = EncryptionContext->GetPublicKey_RSA(InPemKey);
 
 	if (!PublicKey)
 	{
 		UE_LOG(LogJwt, Error,
-			TEXT("[FJwtAlgorithm_RS256::SetPublicKey] RSA public key is invalid."));
+			TEXT("[FJwtAlgorithm_RS256::SetPublicKey] "
+				"RSA public key from string is invalid."));
 
 		return false;
 	}
@@ -66,16 +67,6 @@ bool FJwtAlgorithm_RS256::SetPublicKey(const TArrayView<const uint8> InKey)
 }
 
 
-bool FJwtAlgorithm_RS256::SetPublicKey(const FStringView InKey)
-{
-	TArray<uint8> PublicKeyBytes;
-
-	FJwtUtils::StringViewToBytes(InKey, PublicKeyBytes);
-
-	return SetPublicKey(PublicKeyBytes);
-}
-
-
 void FJwtAlgorithm_RS256::DestroyKey(void* Key)
 {
 	if (!EncryptionContext.IsValid())
diff --git a/Engine/Plugins/Experimental/JWT/Source/JWT/Private/Tests/JwtTests.cpp b/Engine/Plugins/Experimental/JWT/Source/JWT/Private/Tests/JwtTests.cpp
--- a/Engine/Plugins/Experimental/JWT/Source/JWT/Private/Tests/JwtTests.cpp
+++ b/Engine/Plugins/Experimental/JWT/Source/JWT/Private/Tests/JwtTests.cpp
@@ -171,12 +171,7 @@ bool FJwtVerificationTest::RunTest(const FString& Parameters)
 
 	// Set the public key
 	{
-		TArray<uint8> PublicKeyBytes(
-			reinterpret_cast<uint8*>(TCHAR_TO_ANSI(*TestKeys::PUBLICKEY_RSA)),
-			TestKeys::PUBLICKEY_RSA.Len());
-
-
-		if (!Verifier.SetPublicKey(PublicKeyBytes))
+		if (!Verifier.SetPublicKey(TestKeys::PUBLICKEY_RSA))
 		{
 			AddError(TEXT("Could not set public key."));
 			return false;
diff --git a/Engine/Plugins/Experimental/JWT/Source/JWT/Public/JwtAlgorithms.h b/Engine/Plugins/Experimental/JWT/Source/JWT/Public/JwtAlgorithms.h
--- a/Engine/Plugins/Experimental/JWT/Source/JWT/Public/JwtAlgorithms.h
+++ b/Engine/Plugins/Experimental/JWT/Source/JWT/Public/JwtAlgorithms.h
@@ -36,23 +36,14 @@ public:
 
 public:
 
-	/*
-	 * Set the public RSA key.
-	 *
-	 * @param InKey Raw key bytes
-	 *
-	 * @return Whether the operation was successful
-	 */
-	bool SetPublicKey(const TArrayView<const uint8> InKey);
-
 	/**
-	 * Set the public RSA key.
+	 * Set the public RSA key from PEM.
 	 *
-	 * @param InKey Key as string view
+	 * @param InPemKey Key in PEM format as string view
 	 *
 	 * @return Whether the operation was successful
 	 */
-	bool SetPublicKey(const FStringView InKey);
+	bool SetPublicKey(const FStringView InPemKey);
 
 private:
 
diff --git a/Engine/Plugins/Experimental/PlatformCrypto/Source/PlatformCryptoOpenSSL/Private/EncryptionContextOpenSSL.cpp b/Engine/Plugins/Experimental/PlatformCrypto/Source/PlatformCryptoOpenSSL/Private/EncryptionContextOpenSSL.cpp
--- a/Engine/Plugins/Experimental/PlatformCrypto/Source/PlatformCryptoOpenSSL/Private/EncryptionContextOpenSSL.cpp
+++ b/Engine/Plugins/Experimental/PlatformCrypto/Source/PlatformCryptoOpenSSL/Private/EncryptionContextOpenSSL.cpp
@@ -463,9 +463,11 @@ FRSAKeyHandle FEncryptionContextOpenSSL::CreateKey_RSA(const TArrayView<const ui
 	return NewKey;
 }
 
-FRSAKeyHandle FEncryptionContextOpenSSL::GetPublicKey_RSA(const void* Source, const int32 Length)
+FRSAKeyHandle FEncryptionContextOpenSSL::GetPublicKey_RSA(const FStringView PemSource)
 {
-	BIO* KeyBuf = BIO_new_mem_buf(Source, Length);
+	auto ConvertedSource = StringCast<UTF8CHAR>(PemSource.GetData());
+
+	BIO* KeyBuf = BIO_new_mem_buf(ConvertedSource.Get(), ConvertedSource.Length());
 	RSA* Key = PEM_read_bio_RSA_PUBKEY(KeyBuf, 0, 0, 0);
 
 	BIO_free(KeyBuf);
@@ -473,18 +475,6 @@ FRSAKeyHandle FEncryptionContextOpenSSL::GetPublicKey_RSA(const void* Source, co
 	return Key;
 }
 
-FRSAKeyHandle FEncryptionContextOpenSSL::GetPublicKey_RSA(const TArrayView<const uint8> Source)
-{
-	return GetPublicKey_RSA(Source.GetData(), Source.Num());
-}
-
-FRSAKeyHandle FEncryptionContextOpenSSL::GetPublicKey_RSA(const FStringView Source)
-{
-	auto ConvertedSource = StringCast<UTF8CHAR>(Source.GetData());
-
-	return GetPublicKey_RSA(ConvertedSource.Get(), ConvertedSource.Length());
-}
-
 void FEncryptionContextOpenSSL::DestroyKey_RSA(FRSAKeyHandle Key)
 {
 	RSA_free((RSA*)Key);
diff --git a/Engine/Plugins/Experimental/PlatformCrypto/Source/PlatformCryptoOpenSSL/Public/EncryptionContextOpenSSL.h b/Engine/Plugins/Experimental/PlatformCrypto/Source/PlatformCryptoOpenSSL/Public/EncryptionContextOpenSSL.h
--- a/Engine/Plugins/Experimental/PlatformCrypto/Source/PlatformCryptoOpenSSL/Public/EncryptionContextOpenSSL.h
+++ b/Engine/Plugins/Experimental/PlatformCrypto/Source/PlatformCryptoOpenSSL/Public/EncryptionContextOpenSSL.h
@@ -84,9 +84,7 @@ public:
 	bool GenerateKey_RSA(const int32 InNumKeyBits, TArray<uint8>& OutPublicExponent, TArray<uint8>& OutPrivateExponent, TArray<uint8>& OutModulus);
 
 	FRSAKeyHandle CreateKey_RSA(const TArrayView<const uint8> PublicExponent, const TArrayView<const uint8> PrivateExponent, const TArrayView<const uint8> Modulus);
-	FRSAKeyHandle GetPublicKey_RSA(const void* Source, const int32 Length);
-	FRSAKeyHandle GetPublicKey_RSA(const TArrayView<const uint8> Source);
-	FRSAKeyHandle GetPublicKey_RSA(const FStringView Source);
+	FRSAKeyHandle GetPublicKey_RSA(const FStringView PemSource);
 	void DestroyKey_RSA(FRSAKeyHandle Key);
 	int32 GetKeySize_RSA(FRSAKeyHandle Key);
 	int32 GetMaxDataSize_RSA(FRSAKeyHandle Key);
