[PlatformCrypto] Factoring out OpenSSL context utility functions

diff --git a/Engine/Plugins/Experimental/PlatformCrypto/Source/PlatformCryptoOpenSSL/Private/EncryptionContextOpenSSL.cpp b/Engine/Plugins/Experimental/PlatformCrypto/Source/PlatformCryptoOpenSSL/Private/EncryptionContextOpenSSL.cpp
--- a/Engine/Plugins/Experimental/PlatformCrypto/Source/PlatformCryptoOpenSSL/Private/EncryptionContextOpenSSL.cpp
+++ b/Engine/Plugins/Experimental/PlatformCrypto/Source/PlatformCryptoOpenSSL/Private/EncryptionContextOpenSSL.cpp
@@ -3,6 +3,7 @@
 #include "EncryptionContextOpenSSL.h"
 #include "PlatformCryptoAesEncryptorsOpenSSL.h"
 #include "PlatformCryptoAesDecryptorsOpenSSL.h"
+#include "PlatformCryptoUtilsOpenSSL.h"
 
 #include "Containers/StringConv.h"
 
@@ -369,23 +370,6 @@ bool FEncryptionContextOpenSSL::DigestVerify_RS256(const TArrayView<const uint8>
 	return VerifyResult == 1;
 }
 
-// Some platforms were upgraded to OpenSSL 1.1.1 while the others were left on a previous version. There are some minor differences we have to account for
-// in the older version, so declare a handy define that we can use to gate the code
-#if !defined(OPENSSL_VERSION_NUMBER) || OPENSSL_VERSION_NUMBER < 0x10100000L
-#define USE_LEGACY_OPENSSL 1
-#else
-#define USE_LEGACY_OPENSSL 0
-#endif
-
-void BigNumToArray(const int32 InKeySize, const BIGNUM* InNum, TArray<uint8>& OutBytes)
-{
-	int32 NumBytes = BN_num_bytes(InNum);
-	check(NumBytes <= InKeySize);
-	OutBytes.SetNumZeroed(NumBytes);
-
-	BN_bn2bin(InNum, OutBytes.GetData());
-	Algo::Reverse(OutBytes);
-}
 
 bool FEncryptionContextOpenSSL::GenerateKey_RSA(const int32 InNumKeyBits, TArray<uint8>& OutPublicExponent, TArray<uint8>& OutPrivateExponent, TArray<uint8>& OutModulus)
 {
@@ -396,7 +380,6 @@ bool FEncryptionContextOpenSSL::GenerateKey_RSA(const int32 InNumKeyBits, TArray
 	}
 
 	int32 KeySize = InNumKeyBits;
-	int32 KeySizeInBytes = InNumKeyBits / 8;
 
 	RSA* RSAKey = RSA_new();
 	BIGNUM* E = BN_new();
@@ -413,26 +396,15 @@ bool FEncryptionContextOpenSSL::GenerateKey_RSA(const int32 InNumKeyBits, TArray
 	const BIGNUM* PrivateExponent = RSA_get0_d(RSAKey);
 #endif
 
-	BigNumToArray(KeySizeInBytes, PublicModulus, OutModulus);
-	BigNumToArray(KeySizeInBytes, PublicExponent, OutPublicExponent);
-	BigNumToArray(KeySizeInBytes, PrivateExponent, OutPrivateExponent);
+	FPlatformCryptoUtilsOpenSSL::BigNumToArray(PublicModulus, OutModulus);
+	FPlatformCryptoUtilsOpenSSL::BigNumToArray(PublicExponent, OutPublicExponent);
+	FPlatformCryptoUtilsOpenSSL::BigNumToArray(PrivateExponent, OutPrivateExponent);
 
 	RSA_free(RSAKey);
 
 	return true;
 }
 
-static void LoadBinaryIntoBigNum(const uint8* InData, int32 InDataSize, BIGNUM* InBigNum)
-{
-#if USE_LEGACY_OPENSSL
-	TArray<uint8> Bytes(InData, InDataSize);
-	Algo::Reverse(Bytes);
-	BN_bin2bn(Bytes.GetData(), Bytes.Num(), InBigNum);
-#else
-	BN_lebin2bn(InData, InDataSize, InBigNum);
-#endif
-}
-
 FRSAKeyHandle FEncryptionContextOpenSSL::CreateKey_RSA(const TArrayView<const uint8> PublicExponent, const TArrayView<const uint8> PrivateExponent, const TArrayView<const uint8> Modulus)
 {
 	RSA* NewKey = RSA_new();
@@ -443,15 +415,19 @@ FRSAKeyHandle FEncryptionContextOpenSSL::CreateKey_RSA(const TArrayView<const ui
 
 	if (PublicExponent.Num())
 	{
-		LoadBinaryIntoBigNum(PublicExponent.GetData(), PublicExponent.Num(), BN_PublicExponent);
+		FPlatformCryptoUtilsOpenSSL::LoadBinaryIntoBigNum(
+			PublicExponent.GetData(), PublicExponent.Num(), BN_PublicExponent);
 	}
 
 	if (PrivateExponent.Num())
 	{
-		LoadBinaryIntoBigNum(PrivateExponent.GetData(), PrivateExponent.Num(), BN_PrivateExponent);
+		FPlatformCryptoUtilsOpenSSL::LoadBinaryIntoBigNum(
+			PrivateExponent.GetData(), PrivateExponent.Num(), BN_PrivateExponent);
 	}
 
-	LoadBinaryIntoBigNum(Modulus.GetData(), Modulus.Num(), BN_Modulus);
+	FPlatformCryptoUtilsOpenSSL::LoadBinaryIntoBigNum(
+		Modulus.GetData(), Modulus.Num(), BN_Modulus);
+
 #if USE_LEGACY_OPENSSL
 	NewKey->n = BN_Modulus;
 	NewKey->e = BN_PublicExponent;
diff --git a/Engine/Plugins/Experimental/PlatformCrypto/Source/PlatformCryptoOpenSSL/Private/PlatformCryptoUtilsOpenSSL.cpp b/Engine/Plugins/Experimental/PlatformCrypto/Source/PlatformCryptoOpenSSL/Private/PlatformCryptoUtilsOpenSSL.cpp
new file mode 100644
--- /dev/null
+++ b/Engine/Plugins/Experimental/PlatformCrypto/Source/PlatformCryptoOpenSSL/Private/PlatformCryptoUtilsOpenSSL.cpp
@@ -0,0 +1,27 @@
+// Copyright Epic Games, Inc. All Rights Reserved.
+
+#include "PlatformCryptoUtilsOpenSSL.h"
+
+
+void FPlatformCryptoUtilsOpenSSL::BigNumToArray(
+	const BIGNUM* InNum, TArray<uint8>& OutBytes)
+{
+	int32 NumBytes = BN_num_bytes(InNum);
+	OutBytes.SetNumZeroed(NumBytes);
+
+	BN_bn2bin(InNum, OutBytes.GetData());
+	Algo::Reverse(OutBytes);
+}
+
+
+void FPlatformCryptoUtilsOpenSSL::LoadBinaryIntoBigNum(
+	const uint8* InData, int32 InDataSize, BIGNUM* InBigNum)
+{
+#if USE_LEGACY_OPENSSL
+	TArray<uint8> Bytes(InData, InDataSize);
+	Algo::Reverse(Bytes);
+	BN_bin2bn(Bytes.GetData(), Bytes.Num(), InBigNum);
+#else
+	BN_lebin2bn(InData, InDataSize, InBigNum);
+#endif
+}
diff --git a/Engine/Plugins/Experimental/PlatformCrypto/Source/PlatformCryptoOpenSSL/Private/PlatformCryptoUtilsOpenSSL.h b/Engine/Plugins/Experimental/PlatformCrypto/Source/PlatformCryptoOpenSSL/Private/PlatformCryptoUtilsOpenSSL.h
new file mode 100644
--- /dev/null
+++ b/Engine/Plugins/Experimental/PlatformCrypto/Source/PlatformCryptoOpenSSL/Private/PlatformCryptoUtilsOpenSSL.h
@@ -0,0 +1,36 @@
+// Copyright Epic Games, Inc. All Rights Reserved.
+
+#pragma once
+
+#include "CoreMinimal.h"
+
+THIRD_PARTY_INCLUDES_START
+#include <openssl/bn.h>
+THIRD_PARTY_INCLUDES_END
+
+
+/**
+ * Some platforms were upgraded to OpenSSL 1.1.1 while the others
+ * were left on a previous version. There are some minor differences
+ * we have to account for in the older version, so declare a handy
+ * define that we can use to gate the code.
+ */
+#if !defined(OPENSSL_VERSION_NUMBER) || OPENSSL_VERSION_NUMBER < 0x10100000L
+#define USE_LEGACY_OPENSSL 1
+#else
+#define USE_LEGACY_OPENSSL 0
+#endif
+
+
+class FPlatformCryptoUtilsOpenSSL
+{
+
+public:
+
+	static void BigNumToArray(
+		const BIGNUM* InNum, TArray<uint8>& OutBytes);
+
+	static void LoadBinaryIntoBigNum(
+		const uint8* InData, int32 InDataSize, BIGNUM* InBigNum);
+
+};
