[PlatformCrypto] Adding encryption context API documentation

Adding encryption cotnext API documentation.

diff --git a/Engine/Plugins/Experimental/PlatformCrypto/Source/PlatformCryptoOpenSSL/Public/EncryptionContextOpenSSL.h b/Engine/Plugins/Experimental/PlatformCrypto/Source/PlatformCryptoOpenSSL/Public/EncryptionContextOpenSSL.h
--- a/Engine/Plugins/Experimental/PlatformCrypto/Source/PlatformCryptoOpenSSL/Public/EncryptionContextOpenSSL.h
+++ b/Engine/Plugins/Experimental/PlatformCrypto/Source/PlatformCryptoOpenSSL/Public/EncryptionContextOpenSSL.h
@@ -76,28 +76,179 @@ public:
 	TUniquePtr<IPlatformCryptoDecryptor> CreateDecryptor_AES_256_CBC(const TArrayView<const uint8> Key, const TArrayView<const uint8> InitializationVector);
 	TUniquePtr<IPlatformCryptoDecryptor> CreateDecryptor_AES_256_GCM(const TArrayView<const uint8> Key, const TArrayView<const uint8> Nonce, const TArrayView<const uint8> AuthTag);
 
+public:
+
+	/**
+	 * Sign a message with RS256.
+	 *
+	 * @param Message The message
+	 * @param Signature The resulting signature
+	 * @param Key Handle to the RSA key
+	 *
+	 * @return Whether the sign operation was successful
+	 */
 	bool DigestSign_RS256(const TArrayView<const uint8> Message, TArray<uint8>& Signature, FRSAKeyHandle Key);
 
+	/**
+	 * Verify a hashed PS256 message with a signature.
+	 *
+	 * @param Message The hashed message
+	 * @param Signature Signature to verify with
+	 * @param PKCS1Key Handle to the RSA key in PKCS1 format
+	 *
+	 * @return Whether the verify operation was successful
+	 */
 	bool DigestVerify_PS256(const TArrayView<const char> Message, const TArrayView<const uint8> Signature, const TArrayView<const uint8> PKCS1Key);
+
+	/**
+	 * Verify a hashed RS256 message with a signature.
+	 *
+	 * @param Message The hashed message
+	 * @param Signature Signature to verify with
+	 * @param Key Handle to the RSA key
+	 *
+	 * @return Whether the verify operation was successful
+	 */
 	bool DigestVerify_RS256(const TArrayView<const uint8> Message, const TArrayView<const uint8> Signature, FRSAKeyHandle Key);
 
+	/**
+	 * Generate an RSA key with the number of key bits.
+	 *
+	 * @param InNumKeyBits Number of key bits
+	 * @param OutPublicExponent Array for the public exponent
+	 * @param OutPrivateExponent Array for the private exponent
+	 * @param OutModulus Array for the modulus
+	 *
+	 * @return Whether the generate operation was successful
+	 */
 	bool GenerateKey_RSA(const int32 InNumKeyBits, TArray<uint8>& OutPublicExponent, TArray<uint8>& OutPrivateExponent, TArray<uint8>& OutModulus);
 
+	/**
+	 * Create an RSA key from a binary public/private exponent and modulus.
+	 *
+	 * @param PublicExponent Binary key public exponent
+	 * @param PrivateExponent Binary key private exponent
+	 * @param Modulus Binary key modulus
+	 *
+	 * @return The created RSA key
+	 */
 	FRSAKeyHandle CreateKey_RSA(const TArrayView<const uint8> PublicExponent, const TArrayView<const uint8> PrivateExponent, const TArrayView<const uint8> Modulus);
+
+	/**
+	 * Get the RSA public key from a PEM format string.
+	 *
+	 * @param PemSource Key in PEM format (base64)
+	 *
+	 * @return The RSA public key
+	 */
 	FRSAKeyHandle GetPublicKey_RSA(const FStringView PemSource);
+
+	/**
+	 * Destroy and free the RSA key.
+	 *
+	 * @param Key Handle to the key
+	 */
 	void DestroyKey_RSA(FRSAKeyHandle Key);
+
+	/**
+	 * Get the RSA key modulus size in bytes.
+	 *
+	 * @param Key Handle to the key
+	 *
+	 * @return The modulus size in bytes
+	 */
 	int32 GetKeySize_RSA(FRSAKeyHandle Key);
+
+	/**
+	 * Get the maximum data size.
+	 *
+	 * @param Key Handle to the key
+	 *
+	 * @return The maximum data size in bytes
+	 */
 	int32 GetMaxDataSize_RSA(FRSAKeyHandle Key);
+
+	/**
+	 * Encrypt a source with a public RSA key.
+	 *
+	 * @param Source Binary source message
+	 * @param Dest Array for the encrypted data
+	 * @param Key Handle to the key
+	 *
+	 * @return Number of encrypted bytes
+	 */
 	int32 EncryptPublic_RSA(TArrayView<const uint8> Source, TArray<uint8>& Dest, FRSAKeyHandle Key);
+
+	/**
+	 * Encrypt a source with a private RSA key.
+	 *
+	 * @param Source Binary source message
+	 * @param Dest Array for the encrypted data
+	 * @param Key Handle to the key
+	 *
+	 * @return Number of encrypted bytes
+	 */
 	int32 EncryptPrivate_RSA(TArrayView<const uint8> Source, TArray<uint8>& Dest, FRSAKeyHandle Key);
+
+	/**
+	 * Decrypt a source with a public RSA key.
+	 *
+	 * @param Source Binary source message
+	 * @param Dest Array for the encrypted data
+	 * @param Key Handle to the key
+	 *
+	 * @return Number of decrypted bytes
+	 */
 	int32 DecryptPublic_RSA(TArrayView<const uint8> Source, TArray<uint8>& Dest, FRSAKeyHandle Key);
+
+	/**
+	 * Decrypt a source with a private RSA key.
+	 *
+	 * @param Source Binary source message
+	 * @param Dest Array for the encrypted data
+	 * @param Key Handle to the key
+	 *
+	 * @return Number of decrypted bytes
+	 */
 	int32 DecryptPrivate_RSA(TArrayView<const uint8> Source, TArray<uint8>& Dest, FRSAKeyHandle Key);
 
+public:
+
+	/**
+	 * Create random bytes.
+	 *
+	 * @param OutData Array for the random bytes
+	 *
+	 * @return The platform crypto result enum
+	 */
 	EPlatformCryptoResult CreateRandomBytes(const TArrayView<uint8> OutData);
+
+	/**
+	 * Create pseudo random bytes.
+	 *
+	 * @param OutData Array for the pseudo random bytes
+	 *
+	 * @return The platform crypto result enum
+	 */
 	EPlatformCryptoResult CreatePseudoRandomBytes(const TArrayView<uint8> OutData);
 
+public:
+
+	/**
+	 * Create a new SHA256 hasher.
+	 *
+	 * @return Reference to the hasher
+	 */
 	FSHA256Hasher CreateSHA256Hasher();
 
+	/**
+	 * Calculate the SHA256 hash of a message.
+	 *
+	 * @param Source Binary source data
+	 * @param OutHash Array to store the hash in
+	 *
+	 * @return Whether the hash operation was successful
+	 */
 	bool CalcSHA256(const TArrayView<const uint8> Source, TArray<uint8>& OutHash);
 
 };
