PlatformCrypto is required for IasTool to be able to decrypt its input

#rnx

diff --git a/Engine/Plugins/Experimental/PlatformCrypto/PlatformCrypto.uplugin b/Engine/Plugins/Experimental/PlatformCrypto/PlatformCrypto.uplugin
--- a/Engine/Plugins/Experimental/PlatformCrypto/PlatformCrypto.uplugin
+++ b/Engine/Plugins/Experimental/PlatformCrypto/PlatformCrypto.uplugin
@@ -14,7 +14,7 @@
 	"CanContainContent": false,
 	"IsBetaVersion": false,
 	"Installed": false,
-	"SupportedPrograms": [ "UnrealPak", "UnrealPakEx" ],
+	"SupportedPrograms": [ "UnrealPak", "UnrealPakEx", "IasTool" ],
 	"Modules": [
 		{
 			"Name": "PlatformCrypto",
@@ -31,7 +31,7 @@
 			"PlatformDenyList": [
 				"TVOS"
 			],
-			"ProgramAllowList": [ "UnrealPak", "UnrealPakEx" ]
+			"ProgramAllowList": [ "UnrealPak", "UnrealPakEx", "IasTool" ]
 		},
 		{
 			"Name": "PlatformCryptoTypes",
