Getting more headers compiling solo on Win64.

diff --git a/Engine/Plugins/AudioGameplayVolume/Source/AudioGameplayVolume/Public/SubmixOverrideVolumeComponent.h b/Engine/Plugins/AudioGameplayVolume/Source/AudioGameplayVolume/Public/SubmixOverrideVolumeComponent.h
--- a/Engine/Plugins/AudioGameplayVolume/Source/AudioGameplayVolume/Public/SubmixOverrideVolumeComponent.h
+++ b/Engine/Plugins/AudioGameplayVolume/Source/AudioGameplayVolume/Public/SubmixOverrideVolumeComponent.h
@@ -3,10 +3,9 @@
 #pragma once
 
 #include "AudioGameplayVolumeMutator.h"
+#include "Sound/AudioVolume.h"
 #include "SubmixOverrideVolumeComponent.generated.h"
 
-struct FAudioVolumeSubmixOverrideSettings;
-
 /**
  *  FProxyMutator_SubmixOverride - An audio thread representation of Submix Overrides
  */
diff --git a/Engine/Plugins/Developer/Concert/ConcertSync/ConcertSyncClient/Source/ConcertSyncClient/Public/ConcertClientLocalDataStore.h b/Engine/Plugins/Developer/Concert/ConcertSync/ConcertSyncClient/Source/ConcertSyncClient/Public/ConcertClientLocalDataStore.h
--- a/Engine/Plugins/Developer/Concert/ConcertSync/ConcertSyncClient/Source/ConcertSyncClient/Public/ConcertClientLocalDataStore.h
+++ b/Engine/Plugins/Developer/Concert/ConcertSync/ConcertSyncClient/Source/ConcertSyncClient/Public/ConcertClientLocalDataStore.h
@@ -33,11 +33,11 @@ class FConcertDataStore;
  * check(MyStore.FetchAs<uint64>(MyKey).GetValue() == MyValue + 20);
  * @endcode
  */
-class CONCERTSYNCCLIENT_API FConcertClientLocalDataStore
+class FConcertClientLocalDataStore
 {
 public:
 	/** Constructs the data store map. */
-	FConcertClientLocalDataStore();
+	CONCERTSYNCCLIENT_API FConcertClientLocalDataStore();
 
 	/**
 	 * Searches the store for the specified key, if not found, adds a new key/value pair, otherwise,
@@ -146,8 +146,8 @@ public:
 
 private:
 	// Wraps calls to FConcertDataStore to avoid leaking the type publicly.
-	FConcertDataStoreResult InternalFetch(const FName& Key, const FName& TypeName) const;
-	FConcertDataStoreResult InternalStore(const FName& Key, const FName& TypeName, const FConcertSessionSerializedPayload& Value);
+	CONCERTSYNCCLIENT_API FConcertDataStoreResult InternalFetch(const FName& Key, const FName& TypeName) const;
+	CONCERTSYNCCLIENT_API FConcertDataStoreResult InternalStore(const FName& Key, const FName& TypeName, const FConcertSessionSerializedPayload& Value);
 
 	/** Maps keyName/keyValue. */
 	TUniquePtr<FConcertDataStore> DataStore;
diff --git a/Engine/Plugins/Developer/Concert/ConcertUI/ConcertSharedSlate/Source/ConcertSharedSlate/Private/Replication/Editor/View/Tree/SReplicationColumnRow.h b/Engine/Plugins/Developer/Concert/ConcertUI/ConcertSharedSlate/Source/ConcertSharedSlate/Private/Replication/Editor/View/Tree/SReplicationColumnRow.h
--- a/Engine/Plugins/Developer/Concert/ConcertUI/ConcertSharedSlate/Source/ConcertSharedSlate/Private/Replication/Editor/View/Tree/SReplicationColumnRow.h
+++ b/Engine/Plugins/Developer/Concert/ConcertUI/ConcertSharedSlate/Source/ConcertSharedSlate/Private/Replication/Editor/View/Tree/SReplicationColumnRow.h
@@ -4,13 +4,16 @@
 
 #include "ConcertFrontendStyle.h"
 
+#include "Delegates/DelegateCombinations.h"
 #include "Widgets/Layout/SBox.h"
 #include "Widgets/SNullWidget.h"
 #include "Widgets/Views/SExpanderArrow.h"
+#include "Widgets/Views/SListView.h"
 #include "Widgets/Views/STableRow.h"
 #include "Widgets/Views/STableViewBase.h"
 #include "Styling/AppStyle.h"
 #include "Styling/CoreStyle.h"
+#include "Replication/Editor/View/Column/IReplicationTreeColumn.h"
 
 namespace UE::ConcertSharedSlate
 {
diff --git a/Engine/Plugins/Developer/Concert/ConcertUI/ConcertSharedSlate/Source/ConcertSharedSlate/Private/Session/Browser/SNewSessionRow.h b/Engine/Plugins/Developer/Concert/ConcertUI/ConcertSharedSlate/Source/ConcertSharedSlate/Private/Session/Browser/SNewSessionRow.h
--- a/Engine/Plugins/Developer/Concert/ConcertUI/ConcertSharedSlate/Source/ConcertSharedSlate/Private/Session/Browser/SNewSessionRow.h
+++ b/Engine/Plugins/Developer/Concert/ConcertUI/ConcertSharedSlate/Source/ConcertSharedSlate/Private/Session/Browser/SNewSessionRow.h
@@ -4,6 +4,7 @@
 
 #include "ConcertHeaderRowUtils.h"
 #include "Widgets/Views/ITableRow.h"
+#include "Widgets/Views/SListView.h"
 #include "Widgets/Views/STableRow.h"
 
 class STableViewBase;
diff --git a/Engine/Plugins/Developer/Concert/ConcertUI/ConcertSharedSlate/Source/ConcertSharedSlate/Private/Session/Browser/SSaveRestoreSessionRow.h b/Engine/Plugins/Developer/Concert/ConcertUI/ConcertSharedSlate/Source/ConcertSharedSlate/Private/Session/Browser/SSaveRestoreSessionRow.h
--- a/Engine/Plugins/Developer/Concert/ConcertUI/ConcertSharedSlate/Source/ConcertSharedSlate/Private/Session/Browser/SSaveRestoreSessionRow.h
+++ b/Engine/Plugins/Developer/Concert/ConcertUI/ConcertSharedSlate/Source/ConcertSharedSlate/Private/Session/Browser/SSaveRestoreSessionRow.h
@@ -5,6 +5,7 @@
 #include "ConcertHeaderRowUtils.h"
 
 #include "Widgets/Views/ITableRow.h"
+#include "Widgets/Views/SListView.h"
 #include "Widgets/Views/STableRow.h"
 
 class FConcertSessionTreeItem;
diff --git a/Engine/Plugins/Developer/Concert/ConcertUI/ConcertSharedSlate/Source/ConcertSharedSlate/Private/Session/Browser/SSessionGroupRow.h b/Engine/Plugins/Developer/Concert/ConcertUI/ConcertSharedSlate/Source/ConcertSharedSlate/Private/Session/Browser/SSessionGroupRow.h
--- a/Engine/Plugins/Developer/Concert/ConcertUI/ConcertSharedSlate/Source/ConcertSharedSlate/Private/Session/Browser/SSessionGroupRow.h
+++ b/Engine/Plugins/Developer/Concert/ConcertUI/ConcertSharedSlate/Source/ConcertSharedSlate/Private/Session/Browser/SSessionGroupRow.h
@@ -3,6 +3,7 @@
 #pragma once
 
 #include "Widgets/Views/ITableRow.h"
+#include "Widgets/Views/SListView.h"
 #include "Widgets/Views/STableRow.h"
 
 class STableViewBase;
diff --git a/Engine/Plugins/Developer/Concert/ConcertUI/ConcertSharedSlate/Source/ConcertSharedSlate/Private/Session/Browser/SSessionRow.h b/Engine/Plugins/Developer/Concert/ConcertUI/ConcertSharedSlate/Source/ConcertSharedSlate/Private/Session/Browser/SSessionRow.h
--- a/Engine/Plugins/Developer/Concert/ConcertUI/ConcertSharedSlate/Source/ConcertSharedSlate/Private/Session/Browser/SSessionRow.h
+++ b/Engine/Plugins/Developer/Concert/ConcertUI/ConcertSharedSlate/Source/ConcertSharedSlate/Private/Session/Browser/SSessionRow.h
@@ -4,6 +4,7 @@
 
 #include "Widgets/Text/SInlineEditableTextBlock.h"
 #include "Widgets/Views/ITableRow.h"
+#include "Widgets/Views/SListView.h"
 #include "Widgets/Views/STableRow.h"
 
 class STableViewBase;
diff --git a/Engine/Plugins/Editor/ChangelistReview/Source/Private/SSourceControlReviewEntry.h b/Engine/Plugins/Editor/ChangelistReview/Source/Private/SSourceControlReviewEntry.h
--- a/Engine/Plugins/Editor/ChangelistReview/Source/Private/SSourceControlReviewEntry.h
+++ b/Engine/Plugins/Editor/ChangelistReview/Source/Private/SSourceControlReviewEntry.h
@@ -1,8 +1,10 @@
 // Copyright Epic Games, Inc. All Rights Reserved.
 
 #pragma once
+
 #include "SSourceControlReview.h"
 #include "Widgets/Views/ITableRow.h"
+#include "Widgets/Views/SListView.h"
 #include "Widgets/Views/STableRow.h"
 
 enum EBlueprintType : int;
diff --git a/Engine/Plugins/Editor/DataValidation/Source/DataValidation/Public/DataValidationFixers.h b/Engine/Plugins/Editor/DataValidation/Source/DataValidation/Public/DataValidationFixers.h
--- a/Engine/Plugins/Editor/DataValidation/Source/DataValidation/Public/DataValidationFixers.h
+++ b/Engine/Plugins/Editor/DataValidation/Source/DataValidation/Public/DataValidationFixers.h
@@ -5,6 +5,7 @@
 #include "Logging/TokenizedMessage.h"
 #include "Misc/DataValidation/Fixer.h"
 #include "Templates/SharedPointer.h"
+#include "UObject/WeakObjectPtr.h"
 
 namespace UE::DataValidation
 {
diff --git a/Engine/Plugins/Editor/MaterialAnalyzer/Source/Private/SAnalyzedMaterialNodeWidgetItem.h b/Engine/Plugins/Editor/MaterialAnalyzer/Source/Private/SAnalyzedMaterialNodeWidgetItem.h
--- a/Engine/Plugins/Editor/MaterialAnalyzer/Source/Private/SAnalyzedMaterialNodeWidgetItem.h
+++ b/Engine/Plugins/Editor/MaterialAnalyzer/Source/Private/SAnalyzedMaterialNodeWidgetItem.h
@@ -3,6 +3,7 @@
 #pragma once
 
 #include "Widgets/Views/ITableRow.h"
+#include "Widgets/Views/SListView.h"
 #include "Widgets/Views/STableRow.h"
 #include "AnalyzedMaterialNode.h"
 
diff --git a/Engine/Plugins/Editor/PluginBrowser/Source/PluginBrowser/Private/SPluginCategory.h b/Engine/Plugins/Editor/PluginBrowser/Source/PluginBrowser/Private/SPluginCategory.h
--- a/Engine/Plugins/Editor/PluginBrowser/Source/PluginBrowser/Private/SPluginCategory.h
+++ b/Engine/Plugins/Editor/PluginBrowser/Source/PluginBrowser/Private/SPluginCategory.h
@@ -3,6 +3,7 @@
 #pragma once
 
 #include "Widgets/Views/ITableRow.h"
+#include "Widgets/Views/SListView.h"
 #include "Widgets/Views/STableRow.h"
 
 class STableViewBase;
diff --git a/Engine/Plugins/Editor/ProxyLODPlugin/Source/ProxyLOD/Private/ProxyLODGrid2d.h b/Engine/Plugins/Editor/ProxyLODPlugin/Source/ProxyLOD/Private/ProxyLODGrid2d.h
--- a/Engine/Plugins/Editor/ProxyLODPlugin/Source/ProxyLOD/Private/ProxyLODGrid2d.h
+++ b/Engine/Plugins/Editor/ProxyLODPlugin/Source/ProxyLOD/Private/ProxyLODGrid2d.h
@@ -4,6 +4,7 @@
 
 #include "Containers/Array.h"
 #include "Templates/SharedPointer.h"
+#include "Math/IntPoint.h"
 #include "Misc/AssertionMacros.h"
 
 namespace ProxyLOD
diff --git a/Engine/Plugins/Editor/ProxyLODPlugin/Source/ProxyLOD/Private/ProxyLODOpenVDB.h b/Engine/Plugins/Editor/ProxyLODPlugin/Source/ProxyLOD/Private/ProxyLODOpenVDB.h
--- a/Engine/Plugins/Editor/ProxyLODPlugin/Source/ProxyLOD/Private/ProxyLODOpenVDB.h
+++ b/Engine/Plugins/Editor/ProxyLODPlugin/Source/ProxyLOD/Private/ProxyLODOpenVDB.h
@@ -2,6 +2,9 @@
 
 #pragma once
 
+#include "HAL/Platform.h"
+#include "Misc/CoreMiscDefines.h"
+
 THIRD_PARTY_INCLUDES_START
 
 UE_PUSH_MACRO("check")
diff --git a/Engine/Plugins/Editor/ProxyLODPlugin/Source/ProxyLOD/Private/ProxyLODQuadric.h b/Engine/Plugins/Editor/ProxyLODPlugin/Source/ProxyLOD/Private/ProxyLODQuadric.h
--- a/Engine/Plugins/Editor/ProxyLODPlugin/Source/ProxyLOD/Private/ProxyLODQuadric.h
+++ b/Engine/Plugins/Editor/ProxyLODPlugin/Source/ProxyLOD/Private/ProxyLODQuadric.h
@@ -6,6 +6,7 @@
 #pragma once
 
 #include "CoreMinimal.h"
+#include <cmath>
 
 #define SINGULAR_QUADRIC_DET 1e-8
 
diff --git a/Engine/Plugins/Editor/ProxyLODPlugin/Source/ProxyLOD/Private/ProxyLODThreadedWrappers.h b/Engine/Plugins/Editor/ProxyLODPlugin/Source/ProxyLOD/Private/ProxyLODThreadedWrappers.h
--- a/Engine/Plugins/Editor/ProxyLODPlugin/Source/ProxyLOD/Private/ProxyLODThreadedWrappers.h
+++ b/Engine/Plugins/Editor/ProxyLODPlugin/Source/ProxyLOD/Private/ProxyLODThreadedWrappers.h
@@ -1,6 +1,9 @@
 // Copyright Epic Games, Inc. All Rights Reserved.
 #pragma once
 
+#include "HAL/Platform.h"
+#include "ProfilingDebugging/CpuProfilerTrace.h"
+
 THIRD_PARTY_INCLUDES_START
 PRAGMA_DISABLE_DEPRECATION_WARNINGS
 #include <tbb/blocked_range.h>
diff --git a/Engine/Plugins/Editor/ProxyLODPlugin/Source/ProxyLOD/Private/ProxyLODTwoDTriangleUtilities.h b/Engine/Plugins/Editor/ProxyLODPlugin/Source/ProxyLOD/Private/ProxyLODTwoDTriangleUtilities.h
--- a/Engine/Plugins/Editor/ProxyLODPlugin/Source/ProxyLOD/Private/ProxyLODTwoDTriangleUtilities.h
+++ b/Engine/Plugins/Editor/ProxyLODPlugin/Source/ProxyLOD/Private/ProxyLODTwoDTriangleUtilities.h
@@ -2,6 +2,9 @@
 
 #pragma once
 
+#include "HAL/Platform.h"
+#include "Math/UnrealMathUtility.h"
+
 THIRD_PARTY_INCLUDES_START
 #include <array>
 #include <limits>
diff --git a/Engine/Plugins/Editor/ProxyLODPlugin/Source/ProxyLOD/Private/ProxyLODVertexTypes.h b/Engine/Plugins/Editor/ProxyLODPlugin/Source/ProxyLOD/Private/ProxyLODVertexTypes.h
--- a/Engine/Plugins/Editor/ProxyLODPlugin/Source/ProxyLOD/Private/ProxyLODVertexTypes.h
+++ b/Engine/Plugins/Editor/ProxyLODPlugin/Source/ProxyLOD/Private/ProxyLODVertexTypes.h
@@ -2,6 +2,8 @@
 
 #pragma once
 
+#include "Math/Vector.h"
+
 // Implement vertex types that meet the required templated interface for the simplifier.
 
 /**
diff --git a/Engine/Plugins/Enterprise/VariantManager/Source/VariantManager/Private/SDependencyRow.h b/Engine/Plugins/Enterprise/VariantManager/Source/VariantManager/Private/SDependencyRow.h
--- a/Engine/Plugins/Enterprise/VariantManager/Source/VariantManager/Private/SDependencyRow.h
+++ b/Engine/Plugins/Enterprise/VariantManager/Source/VariantManager/Private/SDependencyRow.h
@@ -4,6 +4,7 @@
 
 
 #include "Widgets/Views/ITableRow.h"
+#include "Widgets/Views/SListView.h"
 #include "Widgets/Views/STableRow.h"
 
 class STableViewBase;
diff --git a/Engine/Plugins/Experimental/FullBodyIK/Source/PBIK/Public/Core/PBIKSolver.h b/Engine/Plugins/Experimental/FullBodyIK/Source/PBIK/Public/Core/PBIKSolver.h
--- a/Engine/Plugins/Experimental/FullBodyIK/Source/PBIK/Public/Core/PBIKSolver.h
+++ b/Engine/Plugins/Experimental/FullBodyIK/Source/PBIK/Public/Core/PBIKSolver.h
@@ -197,7 +197,7 @@ struct PBIK_API FPBIKSolverSettings
 };
 
 USTRUCT()
-struct PBIK_API FPBIKSolver
+struct FPBIKSolver
 {
 	GENERATED_BODY()
 
@@ -209,52 +209,52 @@ public:
 		return *this;
 	}
 
-	PBIK::FDebugDraw* GetDebugDraw();
+	PBIK_API PBIK::FDebugDraw* GetDebugDraw();
 
 	//
 	// main runtime functions
 	//
 
-	bool Initialize();
+	PBIK_API bool Initialize();
 
-	void Solve(const FPBIKSolverSettings& Settings);
+	PBIK_API void Solve(const FPBIKSolverSettings& Settings);
 
-	void Reset();
+	PBIK_API void Reset();
 
-	bool IsReadyToSimulate() const;
+	PBIK_API bool IsReadyToSimulate() const;
 
 	//
 	// set input / get output at runtime
 	//
 
-	void SetBoneTransform(const int32 Index, const FTransform& InTransform);
+	PBIK_API void SetBoneTransform(const int32 Index, const FTransform& InTransform);
 
-	PBIK::FBoneSettings* GetBoneSettings(const int32 Index);
+	PBIK_API PBIK::FBoneSettings* GetBoneSettings(const int32 Index);
 
-	void SetEffectorGoal(
+	PBIK_API void SetEffectorGoal(
 		const int32 Index, 
 		const FVector& InPosition, 
 		const FQuat& InRotation, 
 		const PBIK::FEffectorSettings& Settings);
 
-	void GetBoneGlobalTransform(const int32 Index, FTransform& OutTransform);
+	PBIK_API void GetBoneGlobalTransform(const int32 Index, FTransform& OutTransform);
 
 	int32 GetNumBones() const { return Bones.Num(); }
 
-	int32 GetBoneIndex(FName BoneName) const;
+	PBIK_API int32 GetBoneIndex(FName BoneName) const;
 
 	//
 	// pre-init /  setup functions
 	//
 
-	int32 AddBone(
+	PBIK_API int32 AddBone(
 		const FName Name,
 		const int32 ParentIndex,
 		const FVector& InOrigPosition,
 		const FQuat& InOrigRotation,
 		bool bIsSolverRoot);
 
-	int32 AddEffector(const FName BoneName);
+	PBIK_API int32 AddEffector(const FName BoneName);
 	
 private:
 
diff --git a/Engine/Plugins/Experimental/PlatformCrypto/Source/PlatformCryptoOpenSSL/Public/EncryptionContextOpenSSL.h b/Engine/Plugins/Experimental/PlatformCrypto/Source/PlatformCryptoOpenSSL/Public/EncryptionContextOpenSSL.h
--- a/Engine/Plugins/Experimental/PlatformCrypto/Source/PlatformCryptoOpenSSL/Public/EncryptionContextOpenSSL.h
+++ b/Engine/Plugins/Experimental/PlatformCrypto/Source/PlatformCryptoOpenSSL/Public/EncryptionContextOpenSSL.h
@@ -11,7 +11,7 @@ class IPlatformCryptoEncryptor;
 enum class EPlatformCryptoResult;
 
 /** Implementation details for SHA256 computation using OpenSSL */
-struct PLATFORMCRYPTOOPENSSL_API FSHA256HasherOpenSSL final
+struct FSHA256HasherOpenSSL final
 {
 	FSHA256HasherOpenSSL(FSHA256HasherOpenSSL&&) = default;
 	FSHA256HasherOpenSSL& operator=(FSHA256HasherOpenSSL&&) = default;
@@ -20,21 +20,21 @@ struct PLATFORMCRYPTOOPENSSL_API FSHA256HasherOpenSSL final
 	 * Initialize the necessary state to begin computing a message digest.
 	 * It is only necessary to call this function if the object is being used to compute multiple hashes, as the constructor will automatically call it the first time.
 	 */
-	EPlatformCryptoResult Init();
+	PLATFORMCRYPTOOPENSSL_API EPlatformCryptoResult Init();
 
 	/**
 	 * Update the message digest computation with additional bytes
 	 *
 	 * @param InDataBuffer Buffer pointing to the data
 	 */
-	EPlatformCryptoResult Update(const TArrayView<const uint8> InDataBuffer);
+	PLATFORMCRYPTOOPENSSL_API EPlatformCryptoResult Update(const TArrayView<const uint8> InDataBuffer);
 
 	/**
 	 * Finalize the computation of the message digest. After calling this method, Init must be called again if this object is meant to be reused for hashing a new input.
 	 *
 	 * @param OutDataBuffer A buffer that can hold the message digest bytes. Call GetOutputByteLength to determine the necessary size.
 	 */
-	EPlatformCryptoResult Finalize(const TArrayView<uint8> OutDataBuffer);
+	PLATFORMCRYPTOOPENSSL_API EPlatformCryptoResult Finalize(const TArrayView<uint8> OutDataBuffer);
 
 	/**
 	 * The final message digest length in bytes.
@@ -55,26 +55,26 @@ using FSHA256Hasher = FSHA256HasherOpenSSL;
 /**
  * Interface to certain cryptographic algorithms, using OpenSSL to implement them.
  */
-class PLATFORMCRYPTOOPENSSL_API FEncryptionContextOpenSSL
+class FEncryptionContextOpenSSL
 {
 
 public:
 
-	TArray<uint8> Encrypt_AES_256_ECB(const TArrayView<const uint8> Plaintext, const TArrayView<const uint8> Key, EPlatformCryptoResult& OutResult);
-	TArray<uint8> Encrypt_AES_256_CBC(const TArrayView<const uint8> Plaintext, const TArrayView<const uint8> Key, const TArrayView<const uint8> InitializationVector, EPlatformCryptoResult& OutResult);
-	TArray<uint8> Encrypt_AES_256_GCM(const TArrayView<const uint8> Plaintext, const TArrayView<const uint8> Key, const TArrayView<const uint8> Nonce, TArray<uint8>& OutAuthTag, EPlatformCryptoResult& OutResult);
+	PLATFORMCRYPTOOPENSSL_API TArray<uint8> Encrypt_AES_256_ECB(const TArrayView<const uint8> Plaintext, const TArrayView<const uint8> Key, EPlatformCryptoResult& OutResult);
+	PLATFORMCRYPTOOPENSSL_API TArray<uint8> Encrypt_AES_256_CBC(const TArrayView<const uint8> Plaintext, const TArrayView<const uint8> Key, const TArrayView<const uint8> InitializationVector, EPlatformCryptoResult& OutResult);
+	PLATFORMCRYPTOOPENSSL_API TArray<uint8> Encrypt_AES_256_GCM(const TArrayView<const uint8> Plaintext, const TArrayView<const uint8> Key, const TArrayView<const uint8> Nonce, TArray<uint8>& OutAuthTag, EPlatformCryptoResult& OutResult);
 
-	TArray<uint8> Decrypt_AES_256_ECB(const TArrayView<const uint8> Ciphertext, const TArrayView<const uint8> Key, EPlatformCryptoResult& OutResult);
-	TArray<uint8> Decrypt_AES_256_CBC(const TArrayView<const uint8> Ciphertext, const TArrayView<const uint8> Key, const TArrayView<const uint8> InitializationVector, EPlatformCryptoResult& OutResult);
-	TArray<uint8> Decrypt_AES_256_GCM(const TArrayView<const uint8> Ciphertext, const TArrayView<const uint8> Key, const TArrayView<const uint8> Nonce, const TArrayView<const uint8> AuthTag, EPlatformCryptoResult& OutResult);
+	PLATFORMCRYPTOOPENSSL_API TArray<uint8> Decrypt_AES_256_ECB(const TArrayView<const uint8> Ciphertext, const TArrayView<const uint8> Key, EPlatformCryptoResult& OutResult);
+	PLATFORMCRYPTOOPENSSL_API TArray<uint8> Decrypt_AES_256_CBC(const TArrayView<const uint8> Ciphertext, const TArrayView<const uint8> Key, const TArrayView<const uint8> InitializationVector, EPlatformCryptoResult& OutResult);
+	PLATFORMCRYPTOOPENSSL_API TArray<uint8> Decrypt_AES_256_GCM(const TArrayView<const uint8> Ciphertext, const TArrayView<const uint8> Key, const TArrayView<const uint8> Nonce, const TArrayView<const uint8> AuthTag, EPlatformCryptoResult& OutResult);
 
-	TUniquePtr<IPlatformCryptoEncryptor> CreateEncryptor_AES_256_ECB(const TArrayView<const uint8> Key);
-	TUniquePtr<IPlatformCryptoEncryptor> CreateEncryptor_AES_256_CBC(const TArrayView<const uint8> Key, const TArrayView<const uint8> InitializationVector);
-	TUniquePtr<IPlatformCryptoEncryptor> CreateEncryptor_AES_256_GCM(const TArrayView<const uint8> Key, const TArrayView<const uint8> Nonce);
+	PLATFORMCRYPTOOPENSSL_API TUniquePtr<IPlatformCryptoEncryptor> CreateEncryptor_AES_256_ECB(const TArrayView<const uint8> Key);
+	PLATFORMCRYPTOOPENSSL_API TUniquePtr<IPlatformCryptoEncryptor> CreateEncryptor_AES_256_CBC(const TArrayView<const uint8> Key, const TArrayView<const uint8> InitializationVector);
+	PLATFORMCRYPTOOPENSSL_API TUniquePtr<IPlatformCryptoEncryptor> CreateEncryptor_AES_256_GCM(const TArrayView<const uint8> Key, const TArrayView<const uint8> Nonce);
 
-	TUniquePtr<IPlatformCryptoDecryptor> CreateDecryptor_AES_256_ECB(const TArrayView<const uint8> Key);
-	TUniquePtr<IPlatformCryptoDecryptor> CreateDecryptor_AES_256_CBC(const TArrayView<const uint8> Key, const TArrayView<const uint8> InitializationVector);
-	TUniquePtr<IPlatformCryptoDecryptor> CreateDecryptor_AES_256_GCM(const TArrayView<const uint8> Key, const TArrayView<const uint8> Nonce, const TArrayView<const uint8> AuthTag);
+	PLATFORMCRYPTOOPENSSL_API TUniquePtr<IPlatformCryptoDecryptor> CreateDecryptor_AES_256_ECB(const TArrayView<const uint8> Key);
+	PLATFORMCRYPTOOPENSSL_API TUniquePtr<IPlatformCryptoDecryptor> CreateDecryptor_AES_256_CBC(const TArrayView<const uint8> Key, const TArrayView<const uint8> InitializationVector);
+	PLATFORMCRYPTOOPENSSL_API TUniquePtr<IPlatformCryptoDecryptor> CreateDecryptor_AES_256_GCM(const TArrayView<const uint8> Key, const TArrayView<const uint8> Nonce, const TArrayView<const uint8> AuthTag);
 
 public:
 
@@ -87,7 +87,7 @@ public:
 	 *
 	 * @return Whether the sign operation was successful
 	 */
-	bool DigestSign_RS256(const TArrayView<const uint8> Message, TArray<uint8>& Signature, FRSAKeyHandle Key);
+	PLATFORMCRYPTOOPENSSL_API bool DigestSign_RS256(const TArrayView<const uint8> Message, TArray<uint8>& Signature, FRSAKeyHandle Key);
 
 	/**
 	 * Verify a hashed PS256 message with a signature.
@@ -98,7 +98,7 @@ public:
 	 *
 	 * @return Whether the verify operation was successful
 	 */
-	bool DigestVerify_PS256(const TArrayView<const char> Message, const TArrayView<const uint8> Signature, const TArrayView<const uint8> PKCS1Key);
+	PLATFORMCRYPTOOPENSSL_API bool DigestVerify_PS256(const TArrayView<const char> Message, const TArrayView<const uint8> Signature, const TArrayView<const uint8> PKCS1Key);
 
 	/**
 	 * Verify a hashed RS256 message with a signature.
@@ -109,7 +109,7 @@ public:
 	 *
 	 * @return Whether the verify operation was successful
 	 */
-	bool DigestVerify_RS256(const TArrayView<const uint8> Message, const TArrayView<const uint8> Signature, FRSAKeyHandle Key);
+	PLATFORMCRYPTOOPENSSL_API bool DigestVerify_RS256(const TArrayView<const uint8> Message, const TArrayView<const uint8> Signature, FRSAKeyHandle Key);
 
 	/**
 	 * Generate an RSA key with the number of key bits.
@@ -121,7 +121,7 @@ public:
 	 *
 	 * @return Whether the generate operation was successful
 	 */
-	bool GenerateKey_RSA(const int32 InNumKeyBits, TArray<uint8>& OutPublicExponent, TArray<uint8>& OutPrivateExponent, TArray<uint8>& OutModulus);
+	PLATFORMCRYPTOOPENSSL_API bool GenerateKey_RSA(const int32 InNumKeyBits, TArray<uint8>& OutPublicExponent, TArray<uint8>& OutPrivateExponent, TArray<uint8>& OutModulus);
 
 	/**
 	 * Create an RSA key from a binary public/private exponent and modulus.
@@ -132,7 +132,7 @@ public:
 	 *
 	 * @return The created RSA key
 	 */
-	FRSAKeyHandle CreateKey_RSA(const TArrayView<const uint8> PublicExponent, const TArrayView<const uint8> PrivateExponent, const TArrayView<const uint8> Modulus);
+	PLATFORMCRYPTOOPENSSL_API FRSAKeyHandle CreateKey_RSA(const TArrayView<const uint8> PublicExponent, const TArrayView<const uint8> PrivateExponent, const TArrayView<const uint8> Modulus);
 
 	/**
 	 * Get the RSA public key from a PEM format string.
@@ -141,14 +141,14 @@ public:
 	 *
 	 * @return The RSA public key
 	 */
-	FRSAKeyHandle GetPublicKey_RSA(const FStringView PemSource);
+	PLATFORMCRYPTOOPENSSL_API FRSAKeyHandle GetPublicKey_RSA(const FStringView PemSource);
 
 	/**
 	 * Destroy and free the RSA key.
 	 *
 	 * @param Key Handle to the key
 	 */
-	void DestroyKey_RSA(FRSAKeyHandle Key);
+	PLATFORMCRYPTOOPENSSL_API void DestroyKey_RSA(FRSAKeyHandle Key);
 
 	/**
 	 * Get the RSA key modulus size in bytes.
@@ -157,7 +157,7 @@ public:
 	 *
 	 * @return The modulus size in bytes
 	 */
-	int32 GetKeySize_RSA(FRSAKeyHandle Key);
+	PLATFORMCRYPTOOPENSSL_API int32 GetKeySize_RSA(FRSAKeyHandle Key);
 
 	/**
 	 * Get the maximum data size.
@@ -166,7 +166,7 @@ public:
 	 *
 	 * @return The maximum data size in bytes
 	 */
-	int32 GetMaxDataSize_RSA(FRSAKeyHandle Key);
+	PLATFORMCRYPTOOPENSSL_API int32 GetMaxDataSize_RSA(FRSAKeyHandle Key);
 
 	/**
 	 * Encrypt a source with a public RSA key.
@@ -177,7 +177,7 @@ public:
 	 *
 	 * @return Number of encrypted bytes
 	 */
-	int32 EncryptPublic_RSA(TArrayView<const uint8> Source, TArray<uint8>& Dest, FRSAKeyHandle Key);
+	PLATFORMCRYPTOOPENSSL_API int32 EncryptPublic_RSA(TArrayView<const uint8> Source, TArray<uint8>& Dest, FRSAKeyHandle Key);
 
 	/**
 	 * Encrypt a source with a private RSA key.
@@ -188,7 +188,7 @@ public:
 	 *
 	 * @return Number of encrypted bytes
 	 */
-	int32 EncryptPrivate_RSA(TArrayView<const uint8> Source, TArray<uint8>& Dest, FRSAKeyHandle Key);
+	PLATFORMCRYPTOOPENSSL_API int32 EncryptPrivate_RSA(TArrayView<const uint8> Source, TArray<uint8>& Dest, FRSAKeyHandle Key);
 
 	/**
 	 * Decrypt a source with a public RSA key.
@@ -199,7 +199,7 @@ public:
 	 *
 	 * @return Number of decrypted bytes
 	 */
-	int32 DecryptPublic_RSA(TArrayView<const uint8> Source, TArray<uint8>& Dest, FRSAKeyHandle Key);
+	PLATFORMCRYPTOOPENSSL_API int32 DecryptPublic_RSA(TArrayView<const uint8> Source, TArray<uint8>& Dest, FRSAKeyHandle Key);
 
 	/**
 	 * Decrypt a source with a private RSA key.
@@ -210,7 +210,7 @@ public:
 	 *
 	 * @return Number of decrypted bytes
 	 */
-	int32 DecryptPrivate_RSA(TArrayView<const uint8> Source, TArray<uint8>& Dest, FRSAKeyHandle Key);
+	PLATFORMCRYPTOOPENSSL_API int32 DecryptPrivate_RSA(TArrayView<const uint8> Source, TArray<uint8>& Dest, FRSAKeyHandle Key);
 
 public:
 
@@ -221,7 +221,7 @@ public:
 	 *
 	 * @return The platform crypto result enum
 	 */
-	EPlatformCryptoResult CreateRandomBytes(const TArrayView<uint8> OutData);
+	PLATFORMCRYPTOOPENSSL_API EPlatformCryptoResult CreateRandomBytes(const TArrayView<uint8> OutData);
 
 	/**
 	 * Create pseudo random bytes.
@@ -230,7 +230,7 @@ public:
 	 *
 	 * @return The platform crypto result enum
 	 */
-	EPlatformCryptoResult CreatePseudoRandomBytes(const TArrayView<uint8> OutData);
+	PLATFORMCRYPTOOPENSSL_API EPlatformCryptoResult CreatePseudoRandomBytes(const TArrayView<uint8> OutData);
 
 public:
 
@@ -239,7 +239,7 @@ public:
 	 *
 	 * @return Reference to the hasher
 	 */
-	FSHA256Hasher CreateSHA256Hasher();
+	PLATFORMCRYPTOOPENSSL_API FSHA256Hasher CreateSHA256Hasher();
 
 	/**
 	 * Calculate the SHA256 hash of a message.
@@ -249,7 +249,7 @@ public:
 	 *
 	 * @return Whether the hash operation was successful
 	 */
-	bool CalcSHA256(const TArrayView<const uint8> Source, TArray<uint8>& OutHash);
+	PLATFORMCRYPTOOPENSSL_API bool CalcSHA256(const TArrayView<const uint8> Source, TArray<uint8>& OutHash);
 
 };
 
diff --git a/Engine/Plugins/FX/Niagara/Source/Niagara/Public/NiagaraDefines.h b/Engine/Plugins/FX/Niagara/Source/Niagara/Public/NiagaraDefines.h
--- a/Engine/Plugins/FX/Niagara/Source/Niagara/Public/NiagaraDefines.h
+++ b/Engine/Plugins/FX/Niagara/Source/Niagara/Public/NiagaraDefines.h
@@ -2,6 +2,8 @@
 
 #pragma once
 
+#include "HAL/Platform.h"
+
 //#define NIAGARA_NAN_CHECKING 1
 #define NIAGARA_NAN_CHECKING 0
 
diff --git a/Engine/Plugins/FX/Niagara/Source/NiagaraEditor/Private/Widgets/Scalability/SNiagaraSystemResolvedScalabilitySettings.h b/Engine/Plugins/FX/Niagara/Source/NiagaraEditor/Private/Widgets/Scalability/SNiagaraSystemResolvedScalabilitySettings.h
--- a/Engine/Plugins/FX/Niagara/Source/NiagaraEditor/Private/Widgets/Scalability/SNiagaraSystemResolvedScalabilitySettings.h
+++ b/Engine/Plugins/FX/Niagara/Source/NiagaraEditor/Private/Widgets/Scalability/SNiagaraSystemResolvedScalabilitySettings.h
@@ -6,6 +6,7 @@
 #include "NiagaraSystem.h"
 #include "ViewModels/NiagaraSystemViewModel.h"
 #include "Widgets/Layout/SGridPanel.h"
+#include "Widgets/Views/SListView.h"
 #include "Widgets/Views/STableRow.h"
 
 class IDetailTreeNode;
diff --git a/Engine/Plugins/FX/Niagara/Source/NiagaraEditorWidgets/Private/SNiagaraScalabilityPreviewSettings.h b/Engine/Plugins/FX/Niagara/Source/NiagaraEditorWidgets/Private/SNiagaraScalabilityPreviewSettings.h
--- a/Engine/Plugins/FX/Niagara/Source/NiagaraEditorWidgets/Private/SNiagaraScalabilityPreviewSettings.h
+++ b/Engine/Plugins/FX/Niagara/Source/NiagaraEditorWidgets/Private/SNiagaraScalabilityPreviewSettings.h
@@ -6,6 +6,7 @@
 #include "NiagaraSystem.h"
 #include "ViewModels/NiagaraSystemScalabilityViewModel.h"
 #include "Widgets/Layout/SGridPanel.h"
+#include "Widgets/Views/SListView.h"
 #include "Widgets/Views/STreeView.h"
 #include "Widgets/Input/SMenuAnchor.h"
 #include "ViewModels/NiagaraSystemViewModel.h"
diff --git a/Engine/Plugins/Media/MediaIOFramework/Source/MediaIOCore/Public/MediaIOAudioDebug.h b/Engine/Plugins/Media/MediaIOFramework/Source/MediaIOCore/Public/MediaIOAudioDebug.h
--- a/Engine/Plugins/Media/MediaIOFramework/Source/MediaIOCore/Public/MediaIOAudioDebug.h
+++ b/Engine/Plugins/Media/MediaIOFramework/Source/MediaIOCore/Public/MediaIOAudioDebug.h
@@ -25,7 +25,7 @@
 DEFINE_LOG_CATEGORY_STATIC(LogMediaIOAudioDebug, Log, All)
 
 #define MEDIA_IO_DUMP_AUDIO_MONO(Array)\
-	UE::MediaIOAudioDebug::GetSingleton().ProcessAudio(L#Array, Array);
+	UE::MediaIOAudioDebug::GetSingleton().ProcessAudio(TEXT(#Array), Array);
 
 #define MEDIA_IO_DUMP_AUDIO(Name, Buffer, BufferSize, SampleSize, NumChannels)\
 	UE::MediaIOAudioDebug::GetSingleton().ProcessAudio(Name, Buffer, BufferSize, SampleSize, NumChannels);
diff --git a/Engine/Plugins/Media/WmfMedia/Source/WmfMedia/Private/Wmf/WmfMediaReadState.h b/Engine/Plugins/Media/WmfMedia/Source/WmfMedia/Private/Wmf/WmfMediaReadState.h
--- a/Engine/Plugins/Media/WmfMedia/Source/WmfMedia/Private/Wmf/WmfMediaReadState.h
+++ b/Engine/Plugins/Media/WmfMedia/Source/WmfMedia/Private/Wmf/WmfMediaReadState.h
@@ -2,8 +2,15 @@
 
 #pragma once
 
+#include "HAL/PlatformAtomics.h"
+#include "Misc/AssertionMacros.h"
+
 #include "Windows/AllowWindowsPlatformTypes.h"
 
+THIRD_PARTY_INCLUDES_START
+#include <unknwn.h>
+#include <shlwapi.h>
+THIRD_PARTY_INCLUDES_END
 
 /**
  * Implements the state information for asynchronous reads of byte buffer sources.
diff --git a/Engine/Plugins/Media/WmfMedia/Source/WmfMedia/Public/WmfMediaCodec/WmfMediaCodecGenerator.h b/Engine/Plugins/Media/WmfMedia/Source/WmfMedia/Public/WmfMediaCodec/WmfMediaCodecGenerator.h
--- a/Engine/Plugins/Media/WmfMedia/Source/WmfMedia/Public/WmfMediaCodec/WmfMediaCodecGenerator.h
+++ b/Engine/Plugins/Media/WmfMedia/Source/WmfMedia/Public/WmfMediaCodec/WmfMediaCodecGenerator.h
@@ -2,6 +2,7 @@
 
 #pragma once
 
+#include "WmfMediaCommon.h"
 #include "IWmfMediaCodec.h"
 
 template <class T>
diff --git a/Engine/Plugins/Online/VoiceChat/EOSVoiceChat/Source/EOSVoiceChat/Public/EOSVoiceChatLog.h b/Engine/Plugins/Online/VoiceChat/EOSVoiceChat/Source/EOSVoiceChat/Public/EOSVoiceChatLog.h
--- a/Engine/Plugins/Online/VoiceChat/EOSVoiceChat/Source/EOSVoiceChat/Public/EOSVoiceChatLog.h
+++ b/Engine/Plugins/Online/VoiceChat/EOSVoiceChat/Source/EOSVoiceChat/Public/EOSVoiceChatLog.h
@@ -2,6 +2,7 @@
 
 #pragma once
 
+#include "EOSVoiceChat.h" // IWYU pragma: keep
 #include "EOSShared.h" // IWYU pragma: keep
 #include "Logging/LogMacros.h"
 
diff --git a/Engine/Plugins/Runtime/AudioSynesthesia/Source/AudioSynesthesiaCore/Public/ConstantQFactory.h b/Engine/Plugins/Runtime/AudioSynesthesia/Source/AudioSynesthesiaCore/Public/ConstantQFactory.h
--- a/Engine/Plugins/Runtime/AudioSynesthesia/Source/AudioSynesthesiaCore/Public/ConstantQFactory.h
+++ b/Engine/Plugins/Runtime/AudioSynesthesia/Source/AudioSynesthesiaCore/Public/ConstantQFactory.h
@@ -14,7 +14,7 @@ namespace Audio
 	 *
 	 *  Settings for the Constant Q Non-Real-Time Factory.
 	 */
-	struct AUDIOSYNESTHESIACORE_API FConstantQSettings : public IAnalyzerSettings, public FConstantQAnalyzerSettings 
+	struct FConstantQSettings : public IAnalyzerSettings, public FConstantQAnalyzerSettings 
 	{
 	public:
 		/** Time, in seconds, between constant Q frames. */
@@ -31,20 +31,20 @@ namespace Audio
 	/** 
 	 * FLoudnessResult contains the temporal evolution of loudness for the last buffer of audio analyzed.
 	 */
-	class AUDIOSYNESTHESIACORE_API FConstantQResult : public IAnalyzerResult
+	class FConstantQResult : public IAnalyzerResult
 	{
 	public:
 
 		FConstantQResult() {}
 
 		/** Add a single frame of CQT data */
-		void AddFrame(const int32 InChannelIndex, const float InTimestamp, TArrayView<const float> InSpectrum);
+		AUDIOSYNESTHESIACORE_API void AddFrame(const int32 InChannelIndex, const float InTimestamp, TArrayView<const float> InSpectrum);
 
 		/** Returns const reference to FLoudnessDatum array for individual channel. */
-		const TArray<FConstantQFrame>& GetFramesForChannel(const int32 ChannelIdx) const;
+		AUDIOSYNESTHESIACORE_API const TArray<FConstantQFrame>& GetFramesForChannel(const int32 ChannelIdx) const;
 
 		/** Returns the number of channels. */
-		int32 GetNumChannels() const;
+		AUDIOSYNESTHESIACORE_API int32 GetNumChannels() const;
 
 	private:
 		TMap<int32, TArray<FConstantQFrame> > ChannelCQTFrames;
@@ -54,7 +54,7 @@ namespace Audio
 	 *
 	 *  FConstantQWorker computes a FConstantQResult from audio samples.
 	 */
-	class AUDIOSYNESTHESIACORE_API FConstantQWorker : public IAnalyzerWorker
+	class FConstantQWorker : public IAnalyzerWorker
 	{
 	public:
 
@@ -63,7 +63,7 @@ namespace Audio
 		 * InParams are the parameters which describe characteristics of the input audio.
 		 * InAnalyzerSettings are the settings which control various aspects of the algorithm.
 		 */
-		FConstantQWorker(const FAnalyzerParameters& InParams, const FConstantQSettings& InAnalyzerSettings);
+		AUDIOSYNESTHESIACORE_API FConstantQWorker(const FAnalyzerParameters& InParams, const FConstantQSettings& InAnalyzerSettings);
 
 		/**
 		 *  Analyze audio and put results into results pointer.
@@ -71,7 +71,7 @@ namespace Audio
 		 *  InAudio is an array view of audio.
 		 *  OutResult is a pointer to a valid FConstantQResult
 		 */
-		virtual void Analyze(TArrayView<const float> InAudio, IAnalyzerResult* OutResult) override;
+		AUDIOSYNESTHESIACORE_API virtual void Analyze(TArrayView<const float> InAudio, IAnalyzerResult* OutResult) override;
 
 	private:
 		/** Analyze audio with multiple channels interleaved. */
@@ -106,23 +106,23 @@ namespace Audio
 	 *  
 	 *  Factory for creating FConstantQ workers and results
 	 */
-	class AUDIOSYNESTHESIACORE_API FConstantQFactory : public IAnalyzerFactory
+	class FConstantQFactory : public IAnalyzerFactory
 	{
 	public:
 
 		/** Name of this analyzer type. */
-		virtual FName GetName() const override;
+		AUDIOSYNESTHESIACORE_API virtual FName GetName() const override;
 
 		/** Human readable name of this analyzer. */
-		virtual FString GetTitle() const override;
+		AUDIOSYNESTHESIACORE_API virtual FString GetTitle() const override;
 
 		/** Create a new FConstantQResult. */
-		virtual TUniquePtr<IAnalyzerResult> NewResult() const override;
+		AUDIOSYNESTHESIACORE_API virtual TUniquePtr<IAnalyzerResult> NewResult() const override;
 
 		/** Create a new FConstantQWorker 
 		 *
 		 *  InSettings must be a pointer to FConstantQSetting
 		 */
-		virtual TUniquePtr<IAnalyzerWorker> NewWorker(const FAnalyzerParameters& InParams, const IAnalyzerSettings* InSettings) const override;
+		AUDIOSYNESTHESIACORE_API virtual TUniquePtr<IAnalyzerWorker> NewWorker(const FAnalyzerParameters& InParams, const IAnalyzerSettings* InSettings) const override;
 	};
 }
\ No newline at end of file
diff --git a/Engine/Plugins/Runtime/AudioSynesthesia/Source/AudioSynesthesiaCore/Public/ConstantQNRTFactory.h b/Engine/Plugins/Runtime/AudioSynesthesia/Source/AudioSynesthesiaCore/Public/ConstantQNRTFactory.h
--- a/Engine/Plugins/Runtime/AudioSynesthesia/Source/AudioSynesthesiaCore/Public/ConstantQNRTFactory.h
+++ b/Engine/Plugins/Runtime/AudioSynesthesia/Source/AudioSynesthesiaCore/Public/ConstantQNRTFactory.h
@@ -13,7 +13,7 @@ namespace Audio
 	 *
 	 *  Settings for the Constant Q Non-Real-Time Factory.
 	 */
-	class AUDIOSYNESTHESIACORE_API FConstantQNRTSettings : public IAnalyzerNRTSettings, public FConstantQAnalyzerSettings 
+	class FConstantQNRTSettings : public IAnalyzerNRTSettings, public FConstantQAnalyzerSettings 
 	{
 	public:
 
@@ -38,7 +38,7 @@ namespace Audio
 	 * 
 	 *  Contains Constant Q data relating to one audio window.
 	 */
-	struct AUDIOSYNESTHESIACORE_API FConstantQFrame
+	struct FConstantQFrame
 	{
 		/** Audio channel which produced the data. */
 		int32 Channel;
@@ -68,41 +68,41 @@ namespace Audio
 	 *
 	 *  FConstantQNRTResult is a container for the output of the FConstantQNRTWorker.
 	 */
-	class AUDIOSYNESTHESIACORE_API FConstantQNRTResult : public IAnalyzerNRTResult
+	class FConstantQNRTResult : public IAnalyzerNRTResult
 	{
 	public:
 
-		FConstantQNRTResult();
+		AUDIOSYNESTHESIACORE_API FConstantQNRTResult();
 
 		/** Serialize or unserialize object */
-		virtual void Serialize(FArchive& Archive) override;
+		AUDIOSYNESTHESIACORE_API virtual void Serialize(FArchive& Archive) override;
 
 		/** Add a single frame of CQT data */
-		void AddFrame(int32 InChannelIndex, float InTimestamp, TArrayView<const float> InSpectrum);
+		AUDIOSYNESTHESIACORE_API void AddFrame(int32 InChannelIndex, float InTimestamp, TArrayView<const float> InSpectrum);
 
 		/** Returns true if this result contains data for the given channel index. */
-		bool ContainsChannel(int32 InChannelIndex) const;
+		AUDIOSYNESTHESIACORE_API bool ContainsChannel(int32 InChannelIndex) const;
 
 		/** Retrieve the array of frames for a single channel of audio. */
-		const TArray<FConstantQFrame>& GetFramesForChannel(int32 InChannelIndex) const;
+		AUDIOSYNESTHESIACORE_API const TArray<FConstantQFrame>& GetFramesForChannel(int32 InChannelIndex) const;
 
 		/** Retrieve the difference between the maximum and minimum value in the spectrum. */
-		FFloatInterval GetChannelConstantQInterval(int32 InChannelIdx) const;
+		AUDIOSYNESTHESIACORE_API FFloatInterval GetChannelConstantQInterval(int32 InChannelIdx) const;
 
 		/** Retrieve an array of channel indices which exist in this result. */
-		void GetChannels(TArray<int32>& OutChannels) const;
+		AUDIOSYNESTHESIACORE_API void GetChannels(TArray<int32>& OutChannels) const;
 
 		/** Returns the duration of the analyzed audio in seconds */
-		virtual float GetDurationInSeconds() const override;
+		AUDIOSYNESTHESIACORE_API virtual float GetDurationInSeconds() const override;
 
 		/** Sets the duration of the analyzed audio in seconds */
-		void SetDurationInSeconds(float InDuration); 
+		AUDIOSYNESTHESIACORE_API void SetDurationInSeconds(float InDuration);
 
 	 	/** Returns true if FConstantQFrame arrays are sorted in chronologically ascending order via their timestamp.  */
-		bool IsSortedChronologically() const;
+		AUDIOSYNESTHESIACORE_API bool IsSortedChronologically() const;
 
 		/** Sorts FConstantQFrame arrays in chronologically ascnding order via their timestamp.  */
-		void SortChronologically();
+		AUDIOSYNESTHESIACORE_API void SortChronologically();
 
 
 	private:
@@ -120,7 +120,7 @@ namespace Audio
 	 *
 	 *  FConstantQNRTWorker computes a FConstantQNRTResult from audio samples.
 	 */
-	class AUDIOSYNESTHESIACORE_API FConstantQNRTWorker : public IAnalyzerNRTWorker
+	class FConstantQNRTWorker : public IAnalyzerNRTWorker
 	{
 	public:
 
@@ -129,7 +129,7 @@ namespace Audio
 		 * InParams are the parameters which describe characteristics of the input audio.
 		 * InAnalyzerSettings are the settings which control various aspects of the algorithm.
 		 */
-		FConstantQNRTWorker(const FAnalyzerNRTParameters& InParams, const FConstantQNRTSettings& InAnalyzerSettings);
+		AUDIOSYNESTHESIACORE_API FConstantQNRTWorker(const FAnalyzerNRTParameters& InParams, const FConstantQNRTSettings& InAnalyzerSettings);
 
 		/**
 		 *  Analyze audio and put results into results pointer.
@@ -137,14 +137,14 @@ namespace Audio
 		 *  InAudio is an array view of audio.
 		 *  OutResult is a pointer to a valid FConstantQNRTResult
 		 */
-		virtual void Analyze(TArrayView<const float> InAudio, IAnalyzerNRTResult* OutResult) override;
+		AUDIOSYNESTHESIACORE_API virtual void Analyze(TArrayView<const float> InAudio, IAnalyzerNRTResult* OutResult) override;
 
 		/** 
 		 *  Call when analysis of audio asset is complete. 
 		 *
 		 *  OutResult must be a pointer to a valid FConstantQNRTResult. 
 		 */
-		virtual void Finalize(IAnalyzerNRTResult* OutResult) override;
+		AUDIOSYNESTHESIACORE_API virtual void Finalize(IAnalyzerNRTResult* OutResult) override;
 
 	private:
 
@@ -181,24 +181,24 @@ namespace Audio
 	 *  
 	 *  Factory for creating FConstantQNRT workers and results
 	 */
-	class AUDIOSYNESTHESIACORE_API FConstantQNRTFactory : public IAnalyzerNRTFactory
+	class FConstantQNRTFactory : public IAnalyzerNRTFactory
 	{
 	public:
 
 		/** Name of this analyzer type. */
-		virtual FName GetName() const override;
+		AUDIOSYNESTHESIACORE_API virtual FName GetName() const override;
 
 		/** Human readable name of this analyzer. */
-		virtual FString GetTitle() const override;
+		AUDIOSYNESTHESIACORE_API virtual FString GetTitle() const override;
 
 		/** Create a new FConstantQNRTResult. */
-		virtual TUniquePtr<IAnalyzerNRTResult> NewResult() const override;
+		AUDIOSYNESTHESIACORE_API virtual TUniquePtr<IAnalyzerNRTResult> NewResult() const override;
 
 		/** Create a new FConstantQNRTWorker 
 		 *
 		 *  InSettings must be a pointer to FConstantQNRTSetting
 		 */
-		virtual TUniquePtr<IAnalyzerNRTWorker> NewWorker(const FAnalyzerNRTParameters& InParams, const IAnalyzerNRTSettings* InSettings) const override;
+		AUDIOSYNESTHESIACORE_API virtual TUniquePtr<IAnalyzerNRTWorker> NewWorker(const FAnalyzerNRTParameters& InParams, const IAnalyzerNRTSettings* InSettings) const override;
 	};
 }
 
diff --git a/Engine/Plugins/Runtime/AudioSynesthesia/Source/AudioSynesthesiaCore/Public/LoudnessFactory.h b/Engine/Plugins/Runtime/AudioSynesthesia/Source/AudioSynesthesiaCore/Public/LoudnessFactory.h
--- a/Engine/Plugins/Runtime/AudioSynesthesia/Source/AudioSynesthesiaCore/Public/LoudnessFactory.h
+++ b/Engine/Plugins/Runtime/AudioSynesthesia/Source/AudioSynesthesiaCore/Public/LoudnessFactory.h
@@ -12,10 +12,9 @@ namespace Audio
 	/**
 	 * Contains settings for loudness analyzer.
 	 */
-	class AUDIOSYNESTHESIACORE_API FLoudnessSettings : public IAnalyzerSettings, public FLoudnessAnalyzerSettings
+	class FLoudnessSettings : public IAnalyzerSettings, public FLoudnessAnalyzerSettings
 	{
-		public:
-
+	public:
 		/** Number of seconds between loudness measurements */
 		float AnalysisPeriod = 0.01f;
 	};
@@ -34,27 +33,27 @@ namespace Audio
 	/** 
 	 * FLoudnessResult contains the temporal evolution of loudness for the last buffer of audio analyzed.
 	 */
-	class AUDIOSYNESTHESIACORE_API FLoudnessResult : public IAnalyzerResult
+	class FLoudnessResult : public IAnalyzerResult
 	{
 	public:
 		/**
 		 * Denotes the overall loudness as opposed individual channel indices.
 		 */
-		static const int32 ChannelIndexOverall;
+		AUDIOSYNESTHESIACORE_API static const int32 ChannelIndexOverall;
 
 		FLoudnessResult() {}
 
 		/** Appends an FLoudnessDatum to the container. */
-		void Add(const FLoudnessEntry& InDatum);
+		AUDIOSYNESTHESIACORE_API void Add(const FLoudnessEntry& InDatum);
 
 		/** Returns const reference to FLoudnessDatum array for individual channel. */
-		const TArray<FLoudnessEntry>& GetChannelLoudnessArray(int32 ChannelIdx) const;
+		AUDIOSYNESTHESIACORE_API const TArray<FLoudnessEntry>& GetChannelLoudnessArray(int32 ChannelIdx) const;
 
 		/** Returns const reference to FLoudnessDatum array associated with overall loudness. */
-		const TArray<FLoudnessEntry>& GetLoudnessArray() const;
+		AUDIOSYNESTHESIACORE_API const TArray<FLoudnessEntry>& GetLoudnessArray() const;
 
 		/** Returns the number of channels. */
-		int32 GetNumChannels() const;
+		AUDIOSYNESTHESIACORE_API int32 GetNumChannels() const;
 
 	private:
 		TMap<int32, TArray<FLoudnessEntry> > ChannelLoudnessArrays;
@@ -63,14 +62,14 @@ namespace Audio
 	/**
 	 * FLoudnessWorker performs loudness analysis on input sample buffers.
 	 */
-	class AUDIOSYNESTHESIACORE_API FLoudnessWorker : public IAnalyzerWorker
+	class FLoudnessWorker : public IAnalyzerWorker
 	{
 	public:
 		/** Construct a worker */
-		FLoudnessWorker(const FAnalyzerParameters& InParams, const FLoudnessSettings& InAnalyzerSettings);
+		AUDIOSYNESTHESIACORE_API FLoudnessWorker(const FAnalyzerParameters& InParams, const FLoudnessSettings& InAnalyzerSettings);
 
 		/** Analyzes input sample buffer and updates result. */
-		virtual void Analyze(TArrayView<const float> InAudio, IAnalyzerResult* OutResult) override;
+		AUDIOSYNESTHESIACORE_API virtual void Analyze(TArrayView<const float> InAudio, IAnalyzerResult* OutResult) override;
 
 	private:
 
@@ -89,21 +88,21 @@ namespace Audio
 	/**
 	 * Defines the Loudness analyzer and creates related classes.
 	 */
-	class AUDIOSYNESTHESIACORE_API FLoudnessFactory : public IAnalyzerFactory
+	class FLoudnessFactory : public IAnalyzerFactory
 	{
-		public:
+	public:
 
 		/** Name of specific analyzer type. */
-		virtual FName GetName() const override;
+		AUDIOSYNESTHESIACORE_API virtual FName GetName() const override;
 
 		/** Human readable name of analyzer. */
-		virtual FString GetTitle() const override;
+		AUDIOSYNESTHESIACORE_API virtual FString GetTitle() const override;
 
 		/** Creates a new FLoudnessNRTResult */
-		virtual TUniquePtr<IAnalyzerResult> NewResult() const override;
+		AUDIOSYNESTHESIACORE_API virtual TUniquePtr<IAnalyzerResult> NewResult() const override;
 
 		/**  Creates a new FLoudnessWorker. This expects IAnalyzerNRTSettings to be a valid pointer to a FLoudnessSettings object. */
-		virtual TUniquePtr<IAnalyzerWorker> NewWorker(const FAnalyzerParameters& InParams, const IAnalyzerSettings* InSettings) const override;
+		AUDIOSYNESTHESIACORE_API virtual TUniquePtr<IAnalyzerWorker> NewWorker(const FAnalyzerParameters& InParams, const IAnalyzerSettings* InSettings) const override;
 	};
 }
 
diff --git a/Engine/Plugins/Runtime/AudioSynesthesia/Source/AudioSynesthesiaCore/Public/LoudnessNRTFactory.h b/Engine/Plugins/Runtime/AudioSynesthesia/Source/AudioSynesthesiaCore/Public/LoudnessNRTFactory.h
--- a/Engine/Plugins/Runtime/AudioSynesthesia/Source/AudioSynesthesiaCore/Public/LoudnessNRTFactory.h
+++ b/Engine/Plugins/Runtime/AudioSynesthesia/Source/AudioSynesthesiaCore/Public/LoudnessNRTFactory.h
@@ -12,15 +12,15 @@ namespace Audio
 	/**
 	 * Contains settings for loudness analyzer.
 	 */
-	class AUDIOSYNESTHESIACORE_API FLoudnessNRTSettings : public IAnalyzerNRTSettings, public FLoudnessAnalyzerSettings
+	class FLoudnessNRTSettings : public IAnalyzerNRTSettings, public FLoudnessAnalyzerSettings
 	{
-		public:
-			/** Number of seconds between loudness measurements */
-			float AnalysisPeriod;
+	public:
+		/** Number of seconds between loudness measurements */
+		float AnalysisPeriod;
 
-			FLoudnessNRTSettings()
-			:	AnalysisPeriod(0.01f)
-			{}
+		FLoudnessNRTSettings()
+		:	AnalysisPeriod(0.01f)
+		{}
 	};
 
 
@@ -39,49 +39,49 @@ namespace Audio
 	AUDIOSYNESTHESIACORE_API FArchive &operator <<(FArchive& Ar, FLoudnessDatum& Datum);
 
 	/** FLoudnessNRTResult contains the temporal evolution of loudness. */
-	class AUDIOSYNESTHESIACORE_API FLoudnessNRTResult : public IAnalyzerNRTResult
+	class FLoudnessNRTResult : public IAnalyzerNRTResult
 	{
 	public:
 		/** Denotes the overall loudness channel index as opposed individual channel indices. */
-		static const int32 ChannelIndexOverall;
+		AUDIOSYNESTHESIACORE_API static const int32 ChannelIndexOverall;
 
-		FLoudnessNRTResult();
+		AUDIOSYNESTHESIACORE_API FLoudnessNRTResult();
 
 		/** Defines how to serialize result. */
-		virtual void Serialize(FArchive& Archive) override;
+		AUDIOSYNESTHESIACORE_API virtual void Serialize(FArchive& Archive) override;
 
 		/** Appends an FLoudnessDatum to the container. */
-		void Add(const FLoudnessDatum& InDatum);
+		AUDIOSYNESTHESIACORE_API void Add(const FLoudnessDatum& InDatum);
 
 		/** Returns true if this object data for the given channel index */
-		bool ContainsChannel(int32 InChannelIndex) const;
+		AUDIOSYNESTHESIACORE_API bool ContainsChannel(int32 InChannelIndex) const;
 
 		/** Returns const reference to FLoudnessDatum array for individual channel. */
-		const TArray<FLoudnessDatum>& GetChannelLoudnessArray(int32 ChannelIdx) const;
+		AUDIOSYNESTHESIACORE_API const TArray<FLoudnessDatum>& GetChannelLoudnessArray(int32 ChannelIdx) const;
 
 		/** Returns const reference to FLoudnessDatum array associated with overall loudness. */
-		const TArray<FLoudnessDatum>& GetLoudnessArray() const;
+		AUDIOSYNESTHESIACORE_API const TArray<FLoudnessDatum>& GetLoudnessArray() const;
 
 		/** Returns range in dB of overall loudness result given the noise floor. */
-		float GetLoudnessRange(float InNoiseFloor) const;
+		AUDIOSYNESTHESIACORE_API float GetLoudnessRange(float InNoiseFloor) const;
 
 		/** Returns range in dB of loudness result given the noise floor. */
-		float GetChannelLoudnessRange(int32 InChannelIdx, float InNoiseFloor) const;
+		AUDIOSYNESTHESIACORE_API float GetChannelLoudnessRange(int32 InChannelIdx, float InNoiseFloor) const;
 
 		/** Returns the channel indices available in result. */
-		void GetChannels(TArray<int32>& OutChannels) const;
+		AUDIOSYNESTHESIACORE_API void GetChannels(TArray<int32>& OutChannels) const;
 
 		/** Gets the duration of the analyzed audio. */
-		virtual float GetDurationInSeconds() const override;
+		AUDIOSYNESTHESIACORE_API virtual float GetDurationInSeconds() const override;
 
 		/** Sets the duration of the analyzed audio. */
-		void SetDurationInSeconds(float InDuration);
+		AUDIOSYNESTHESIACORE_API void SetDurationInSeconds(float InDuration);
 
 		/** Returns true if FLoudnessDatum arrays are sorted in chronologically ascending order via their timestamp. */
-		bool IsSortedChronologically() const;
+		AUDIOSYNESTHESIACORE_API bool IsSortedChronologically() const;
 
 		/** Sorts FLoudnessDatum arrays in chronologically ascending order via their timestamp. */
-		void SortChronologically();
+		AUDIOSYNESTHESIACORE_API void SortChronologically();
 
 	private:
 		float DurationInSeconds;
@@ -93,21 +93,21 @@ namespace Audio
 	/** 
 	 * FLoudnessNRTWorker performs loudness analysis on input sample buffers.
 	 */
-	class AUDIOSYNESTHESIACORE_API FLoudnessNRTWorker : public IAnalyzerNRTWorker
+	class FLoudnessNRTWorker : public IAnalyzerNRTWorker
 	{
 	public:
 		/** Construct a worker */
-		FLoudnessNRTWorker(const FAnalyzerNRTParameters& InParams, const FLoudnessNRTSettings& InAnalyzerSettings);
+		AUDIOSYNESTHESIACORE_API FLoudnessNRTWorker(const FAnalyzerNRTParameters& InParams, const FLoudnessNRTSettings& InAnalyzerSettings);
 
 		/**
 		 * Analyzes input sample buffer and updates result. 
 		 */
-		virtual void Analyze(TArrayView<const float> InAudio, IAnalyzerNRTResult* OutResult) override;
+		AUDIOSYNESTHESIACORE_API virtual void Analyze(TArrayView<const float> InAudio, IAnalyzerNRTResult* OutResult) override;
 
 		/**
 		 * Call when all audio data has been analyzed. 
 		 */
-		virtual void Finalize(IAnalyzerNRTResult* OutResult) override;
+		AUDIOSYNESTHESIACORE_API virtual void Finalize(IAnalyzerNRTResult* OutResult) override;
 
 	private:
 
@@ -127,21 +127,21 @@ namespace Audio
 	/**
 	 * Defines the LoudnessNRT analyzer and creates related classes.
 	 */
-	class AUDIOSYNESTHESIACORE_API FLoudnessNRTFactory : public IAnalyzerNRTFactory
+	class FLoudnessNRTFactory : public IAnalyzerNRTFactory
 	{
-		public:
+	public:
 
 		/** Name of specific analyzer type. */
-		virtual FName GetName() const override;
+		AUDIOSYNESTHESIACORE_API virtual FName GetName() const override;
 
 		/** Human readable name of analyzer. */
-		virtual FString GetTitle() const override;
+		AUDIOSYNESTHESIACORE_API virtual FString GetTitle() const override;
 
 		/** Creates a new FLoudnessNRTResult */
-		virtual TUniquePtr<IAnalyzerNRTResult> NewResult() const override;
+		AUDIOSYNESTHESIACORE_API virtual TUniquePtr<IAnalyzerNRTResult> NewResult() const override;
 
 		/** Creates a new FLoudnessNRTWorker. This expects IAnalyzerNRTSettings to be a valid pointer to a FLoudnessNRTSettings object. */
-		virtual TUniquePtr<IAnalyzerNRTWorker> NewWorker(const FAnalyzerNRTParameters& InParams, const IAnalyzerNRTSettings* InSettings) const override;
+		AUDIOSYNESTHESIACORE_API virtual TUniquePtr<IAnalyzerNRTWorker> NewWorker(const FAnalyzerNRTParameters& InParams, const IAnalyzerNRTSettings* InSettings) const override;
 	};
 }
 
diff --git a/Engine/Plugins/Runtime/AudioSynesthesia/Source/AudioSynesthesiaCore/Public/MeterFactory.h b/Engine/Plugins/Runtime/AudioSynesthesia/Source/AudioSynesthesiaCore/Public/MeterFactory.h
--- a/Engine/Plugins/Runtime/AudioSynesthesia/Source/AudioSynesthesiaCore/Public/MeterFactory.h
+++ b/Engine/Plugins/Runtime/AudioSynesthesia/Source/AudioSynesthesiaCore/Public/MeterFactory.h
@@ -12,7 +12,7 @@ namespace Audio
 	/**
 	 * Contains settings for meter  analyzer.
 	 */
-	class AUDIOSYNESTHESIACORE_API FMeterSettings : public IAnalyzerSettings, public FMeterAnalyzerSettings
+	class FMeterSettings : public IAnalyzerSettings, public FMeterAnalyzerSettings
 	{	
 	public:
 		/** Number of seconds between loudness measurements */
@@ -35,27 +35,27 @@ namespace Audio
 	/** 
 	 * FMeterResult contains the temporal evolution of loudness for the last buffer of audio analyzed.
 	 */
-	class AUDIOSYNESTHESIACORE_API FMeterResult : public IAnalyzerResult
+	class FMeterResult : public IAnalyzerResult
 	{
 	public:
 		/**
 		 * Denotes the overall loudness as opposed individual channel indices.
 		 */
-		static const int32 ChannelIndexOverall;
+		AUDIOSYNESTHESIACORE_API static const int32 ChannelIndexOverall;
 
 		FMeterResult() {}
 
 		/** Appends an FMeterEntry to the container. */
-		void Add(const FMeterEntry& InEntry);
+		AUDIOSYNESTHESIACORE_API void Add(const FMeterEntry& InEntry);
 
 		/** Returns const reference to FMeterEntry array for individual channel. */
-		const TArray<FMeterEntry>& GetChannelMeterArray(int32 ChannelIdx) const;
+		AUDIOSYNESTHESIACORE_API const TArray<FMeterEntry>& GetChannelMeterArray(int32 ChannelIdx) const;
 
 		/** Returns const reference to FMeterEntry array associated with overall loudness. */
-		const TArray<FMeterEntry>& GetMeterArray() const;
+		AUDIOSYNESTHESIACORE_API const TArray<FMeterEntry>& GetMeterArray() const;
 
 		/** Returns the number of channels. */
-		int32 GetNumChannels() const;
+		AUDIOSYNESTHESIACORE_API int32 GetNumChannels() const;
 
 	private:
 		float DurationInSeconds = 0.0f;
@@ -65,14 +65,14 @@ namespace Audio
 	/**
 	 * FMeterWorker performs meter analysis on input sample buffers.
 	 */
-	class AUDIOSYNESTHESIACORE_API FMeterWorker : public IAnalyzerWorker
+	class FMeterWorker : public IAnalyzerWorker
 	{
 	public:
 		/** Construct a worker */
-		FMeterWorker(const FAnalyzerParameters& InParams, const FMeterSettings& InAnalyzerSettings);
+		AUDIOSYNESTHESIACORE_API FMeterWorker(const FAnalyzerParameters& InParams, const FMeterSettings& InAnalyzerSettings);
 
 		/** Analyzes input sample buffer and updates result. */
-		virtual void Analyze(TArrayView<const float> InAudio, IAnalyzerResult* OutResult) override;
+		AUDIOSYNESTHESIACORE_API virtual void Analyze(TArrayView<const float> InAudio, IAnalyzerResult* OutResult) override;
 
 	private:
 
@@ -88,21 +88,20 @@ namespace Audio
 	/**
 	 * Defines the meter analyzer and creates related classes.
 	 */
-	class AUDIOSYNESTHESIACORE_API FMeterFactory : public IAnalyzerFactory
+	class FMeterFactory : public IAnalyzerFactory
 	{
-		public:
-
+	public:
 		/** Name of specific analyzer type. */
-		virtual FName GetName() const override;
+		AUDIOSYNESTHESIACORE_API virtual FName GetName() const override;
 
 		/** Human readable name of analyzer. */
-		virtual FString GetTitle() const override;
+		AUDIOSYNESTHESIACORE_API virtual FString GetTitle() const override;
 
 		/** Creates a new FLoudnessNRTResult */
-		virtual TUniquePtr<IAnalyzerResult> NewResult() const override;
+		AUDIOSYNESTHESIACORE_API virtual TUniquePtr<IAnalyzerResult> NewResult() const override;
 
 		/**  Creates a new FLoudnessWorker. This expects IAnalyzerNRTSettings to be a valid pointer to a FLoudnessSettings object. */
-		virtual TUniquePtr<IAnalyzerWorker> NewWorker(const FAnalyzerParameters& InParams, const IAnalyzerSettings* InSettings) const override;
+		AUDIOSYNESTHESIACORE_API virtual TUniquePtr<IAnalyzerWorker> NewWorker(const FAnalyzerParameters& InParams, const IAnalyzerSettings* InSettings) const override;
 	};
 }
 
diff --git a/Engine/Plugins/Runtime/AudioSynesthesia/Source/AudioSynesthesiaCore/Public/SynesthesiaSpectrumAnalysisFactory.h b/Engine/Plugins/Runtime/AudioSynesthesia/Source/AudioSynesthesiaCore/Public/SynesthesiaSpectrumAnalysisFactory.h
--- a/Engine/Plugins/Runtime/AudioSynesthesia/Source/AudioSynesthesiaCore/Public/SynesthesiaSpectrumAnalysisFactory.h
+++ b/Engine/Plugins/Runtime/AudioSynesthesia/Source/AudioSynesthesiaCore/Public/SynesthesiaSpectrumAnalysisFactory.h
@@ -12,7 +12,7 @@ namespace Audio
 	/**
 	 * Contains settings for Spectrum analyzer.
 	 */
-	class AUDIOSYNESTHESIACORE_API FSynesthesiaSpectrumAnalysisSettings : public IAnalyzerSettings, public FSynesthesiaSpectrumAnalyzerSettings
+	class FSynesthesiaSpectrumAnalysisSettings : public IAnalyzerSettings, public FSynesthesiaSpectrumAnalyzerSettings
 	{	
 	public:
 		/** Number of seconds between spectrum measurements */
@@ -32,19 +32,19 @@ namespace Audio
 	/** 
 	 * FSynesthesiaSpectrumResult contains the temporal evolution of loudness for the last buffer of audio analyzed.
 	 */
-	class AUDIOSYNESTHESIACORE_API FSynesthesiaSpectrumResult : public IAnalyzerResult
+	class FSynesthesiaSpectrumResult : public IAnalyzerResult
 	{
 	public:
 		FSynesthesiaSpectrumResult() {}
 
 		/** Appends an FSynesthesiaSpectrumEntry to the container. */
-		void Add(FSynesthesiaSpectrumEntry&& InEntry);
+		AUDIOSYNESTHESIACORE_API void Add(FSynesthesiaSpectrumEntry&& InEntry);
 
 		/** Returns const reference to FSynesthesiaSpectrumEntry array for individual channel. */
-		const TArray<FSynesthesiaSpectrumEntry>& GetChannelSpectrumArray(int32 ChannelIdx) const;
+		AUDIOSYNESTHESIACORE_API const TArray<FSynesthesiaSpectrumEntry>& GetChannelSpectrumArray(int32 ChannelIdx) const;
 
 		/** Returns the number of channels. */
-		int32 GetNumChannels() const;
+		AUDIOSYNESTHESIACORE_API int32 GetNumChannels() const;
 
 	private:
 		float DurationInSeconds = 0.0f;
@@ -54,14 +54,14 @@ namespace Audio
 	/**
 	 * FSynesthesiaSpectrumWorker performs Spectrum analysis on input sample buffers.
 	 */
-	class AUDIOSYNESTHESIACORE_API FSynesthesiaSpectrumAnalysisWorker : public IAnalyzerWorker
+	class FSynesthesiaSpectrumAnalysisWorker : public IAnalyzerWorker
 	{
 	public:
 		/** Construct a worker */
-		FSynesthesiaSpectrumAnalysisWorker(const FAnalyzerParameters& InParams, const FSynesthesiaSpectrumAnalysisSettings& InAnalyzerSettings);
+		AUDIOSYNESTHESIACORE_API FSynesthesiaSpectrumAnalysisWorker(const FAnalyzerParameters& InParams, const FSynesthesiaSpectrumAnalysisSettings& InAnalyzerSettings);
 
 		/** Analyzes input sample buffer and updates result. */
-		virtual void Analyze(TArrayView<const float> InAudio, IAnalyzerResult* OutResult) override;
+		AUDIOSYNESTHESIACORE_API virtual void Analyze(TArrayView<const float> InAudio, IAnalyzerResult* OutResult) override;
 
 	private:
 
@@ -85,21 +85,21 @@ namespace Audio
 	/**
 	 * Defines the Spectrum analyzer and creates related classes.
 	 */
-	class AUDIOSYNESTHESIACORE_API FSynesthesiaSpectrumAnalysisFactory : public IAnalyzerFactory
+	class FSynesthesiaSpectrumAnalysisFactory : public IAnalyzerFactory
 	{
-		public:
+	public:
 
 		/** Name of specific analyzer type. */
-		virtual FName GetName() const override;
+		AUDIOSYNESTHESIACORE_API virtual FName GetName() const override;
 
 		/** Human readable name of analyzer. */
-		virtual FString GetTitle() const override;
+		AUDIOSYNESTHESIACORE_API virtual FString GetTitle() const override;
 
 		/** Creates a new FSynesthesiaSpectrumResult */
-		virtual TUniquePtr<IAnalyzerResult> NewResult() const override;
+		AUDIOSYNESTHESIACORE_API virtual TUniquePtr<IAnalyzerResult> NewResult() const override;
 
 		/**  Creates a new FSynesthesiaSpectrumAnalysisWorker. This expects IAnalyzerNRTSettings to be a valid pointer to a FSynesthesiaSpectrumAnalysisSettings object. */
-		virtual TUniquePtr<IAnalyzerWorker> NewWorker(const FAnalyzerParameters& InParams, const IAnalyzerSettings* InSettings) const override;
+		AUDIOSYNESTHESIACORE_API virtual TUniquePtr<IAnalyzerWorker> NewWorker(const FAnalyzerParameters& InParams, const IAnalyzerSettings* InSettings) const override;
 	};
 }
 
diff --git a/Engine/Plugins/Runtime/AudioWidgets/Source/AudioWidgets/Public/AudioMaterialSlate/SAudioMaterialButton.h b/Engine/Plugins/Runtime/AudioWidgets/Source/AudioWidgets/Public/AudioMaterialSlate/SAudioMaterialButton.h
--- a/Engine/Plugins/Runtime/AudioWidgets/Source/AudioWidgets/Public/AudioMaterialSlate/SAudioMaterialButton.h
+++ b/Engine/Plugins/Runtime/AudioWidgets/Source/AudioWidgets/Public/AudioMaterialSlate/SAudioMaterialButton.h
@@ -2,6 +2,7 @@
 
 #pragma once
 
+#include "AudioMaterialSlate/AudioMaterialSlateTypes.h"
 #include "Framework/SlateDelegates.h"
 #include "Styling/ISlateStyle.h"
 #include "Styling/SlateWidgetStyleAsset.h"
diff --git a/Engine/Plugins/Runtime/DataRegistry/Source/DataRegistryEditor/Private/SDataRegistryListViewRow.h b/Engine/Plugins/Runtime/DataRegistry/Source/DataRegistryEditor/Private/SDataRegistryListViewRow.h
--- a/Engine/Plugins/Runtime/DataRegistry/Source/DataRegistryEditor/Private/SDataRegistryListViewRow.h
+++ b/Engine/Plugins/Runtime/DataRegistry/Source/DataRegistryEditor/Private/SDataRegistryListViewRow.h
@@ -3,6 +3,7 @@
 #pragma once
 
 #include "Widgets/Views/ITableRow.h"
+#include "Widgets/Views/SListView.h"
 #include "Widgets/Views/STableRow.h"
 #include "DataTableEditorUtils.h"
 
diff --git a/Engine/Plugins/Runtime/GeometryCache/Source/GeometryCache/Private/GeometryCacheConstantTopologyWriter.cpp b/Engine/Plugins/Runtime/GeometryCache/Source/GeometryCache/Private/GeometryCacheConstantTopologyWriter.cpp
--- a/Engine/Plugins/Runtime/GeometryCache/Source/GeometryCache/Private/GeometryCacheConstantTopologyWriter.cpp
+++ b/Engine/Plugins/Runtime/GeometryCache/Source/GeometryCache/Private/GeometryCacheConstantTopologyWriter.cpp
@@ -222,6 +222,8 @@ namespace UE::GeometryCacheHelpers
 		Cache->ClearForReimporting();
 		Cache->Materials.Empty();
 	}
+
+	FGeometryCacheConstantTopologyWriter::~FGeometryCacheConstantTopologyWriter() = default;
 	
 	using FTrackWriter = FGeometryCacheConstantTopologyWriter::FTrackWriter;
 	FTrackWriter& FGeometryCacheConstantTopologyWriter::AddTrackWriter(FName TrackName)
@@ -269,6 +271,8 @@ namespace UE::GeometryCacheHelpers
 		Track = TStrongObjectPtr(NewObject<UGeometryCacheTrackStreamable>(CachePtr, TrackName, RF_Public));
 	}
 
+	FGeometryCacheConstantTopologyWriter::FTrackWriter::~FTrackWriter() = default;
+
 	bool FGeometryCacheConstantTopologyWriter::FTrackWriter::WriteAndClose(TArrayView<TArray<FVector3f>> PositionsToMoveFrom)
 	{
 		if (!Owner->Cache.IsValid() || !Track.IsValid())
diff --git a/Engine/Plugins/Runtime/GeometryCache/Source/GeometryCache/Public/GeometryCacheConstantTopologyWriter.h b/Engine/Plugins/Runtime/GeometryCache/Source/GeometryCache/Public/GeometryCacheConstantTopologyWriter.h
--- a/Engine/Plugins/Runtime/GeometryCache/Source/GeometryCache/Public/GeometryCacheConstantTopologyWriter.h
+++ b/Engine/Plugins/Runtime/GeometryCache/Source/GeometryCache/Public/GeometryCacheConstantTopologyWriter.h
@@ -30,10 +30,10 @@ namespace UE::GeometryCacheHelpers
 	 * 	FGeometryCacheConstantTopologyWriter::FTrackWriter& TrackWriter = Writer.AddGetTrackWriter(); // Second track
 	 * 	...
 	 */
-	class GEOMETRYCACHE_API FGeometryCacheConstantTopologyWriter
+	class FGeometryCacheConstantTopologyWriter
 	{
 	public:
-		struct GEOMETRYCACHE_API FConfig
+		struct FConfig
 		{
 			float FPS = 30.0f;
 			float PositionPrecision = 0.001f;
@@ -44,31 +44,27 @@ namespace UE::GeometryCacheHelpers
 		 * Construct a new FGeometryCacheConstantTopologyWriter object. This will remove all existing tracks from the cache.
 		 * @param OutCache 
 		 */
-		FGeometryCacheConstantTopologyWriter(UGeometryCache& OutCache);
-		FGeometryCacheConstantTopologyWriter(UGeometryCache& OutCache, const FConfig& Config);
-		~FGeometryCacheConstantTopologyWriter() = default;
+		GEOMETRYCACHE_API FGeometryCacheConstantTopologyWriter(UGeometryCache& OutCache);
+		GEOMETRYCACHE_API FGeometryCacheConstantTopologyWriter(UGeometryCache& OutCache, const FConfig& Config);
+		GEOMETRYCACHE_API ~FGeometryCacheConstantTopologyWriter();
 
 		FGeometryCacheConstantTopologyWriter(const FGeometryCacheConstantTopologyWriter&) = delete;
 		FGeometryCacheConstantTopologyWriter& operator=(const FGeometryCacheConstantTopologyWriter&) = delete;
-		FGeometryCacheConstantTopologyWriter(FGeometryCacheConstantTopologyWriter&&) = default;
-		FGeometryCacheConstantTopologyWriter& operator=(FGeometryCacheConstantTopologyWriter&&) = default;
 
-		struct GEOMETRYCACHE_API FFrameData
+		struct FFrameData
 		{
 			TArray<FVector3f> Positions;
 			TArray<FVector3f> Normals;
 			TArray<FVector3f> TangentsX;
 		};
 
-		struct GEOMETRYCACHE_API FTrackWriter
+		struct FTrackWriter
 		{
-			FTrackWriter(FGeometryCacheConstantTopologyWriter& InOwner, FName TrackName = FName());
-			~FTrackWriter() = default;
+			GEOMETRYCACHE_API FTrackWriter(FGeometryCacheConstantTopologyWriter& InOwner, FName TrackName = FName());
+			GEOMETRYCACHE_API ~FTrackWriter();
 
 			FTrackWriter(const FTrackWriter&) = delete;
 			FTrackWriter& operator=(const FTrackWriter&) = delete;
-			FTrackWriter(FTrackWriter&&) = default;
-			FTrackWriter& operator=(FTrackWriter&&) = default;
 
 			TArray<uint32> Indices;
 			TArray<FVector2f> UVs;
@@ -85,7 +81,7 @@ namespace UE::GeometryCacheHelpers
 			 * The number of vertices must be the same for all frames.
 			 * @return true if successfully write data and close the track writer.
 			 */
-			bool WriteAndClose(TArrayView<TArray<FVector3f>> PositionsToMoveFrom);
+			GEOMETRYCACHE_API bool WriteAndClose(TArrayView<TArray<FVector3f>> PositionsToMoveFrom);
 
 			/**
 			 * Move the frame data to the cache track and close the TrackWriter.
@@ -93,17 +89,17 @@ namespace UE::GeometryCacheHelpers
 			 * Normals and tangents are optional. If they are not provided, the track will compute them.
 			 * Normals and tangents must have the same size as positions.
 			 */
-			bool WriteAndClose(TArrayView<FFrameData> FramesToMoveFrom);
+			GEOMETRYCACHE_API bool WriteAndClose(TArrayView<FFrameData> FramesToMoveFrom);
 		private:
 			TStrongObjectPtr<UGeometryCacheTrack> Track;
 			FGeometryCacheConstantTopologyWriter* Owner = nullptr;
 		};
 
-		FTrackWriter& AddTrackWriter(FName TrackName = FName());
-		FTrackWriter& GetTrackWriter(int32 Index);
-		int32 GetNumTracks() const;
-		void AddMaterials(const TArray<TObjectPtr<UMaterialInterface>>& Materials);
-		int32 GetNumMaterials() const;
+		GEOMETRYCACHE_API FTrackWriter& AddTrackWriter(FName TrackName = FName());
+		GEOMETRYCACHE_API FTrackWriter& GetTrackWriter(int32 Index);
+		GEOMETRYCACHE_API int32 GetNumTracks() const;
+		GEOMETRYCACHE_API void AddMaterials(const TArray<TObjectPtr<UMaterialInterface>>& Materials);
+		GEOMETRYCACHE_API int32 GetNumMaterials() const;
 
 	private:
 		TStrongObjectPtr<UGeometryCache> Cache;
diff --git a/Engine/Plugins/Runtime/MeshModelingToolset/Source/ModelingComponentsEditorOnly/Public/AssetUtils/CreateSkeletalMeshUtil.h b/Engine/Plugins/Runtime/MeshModelingToolset/Source/ModelingComponentsEditorOnly/Public/AssetUtils/CreateSkeletalMeshUtil.h
--- a/Engine/Plugins/Runtime/MeshModelingToolset/Source/ModelingComponentsEditorOnly/Public/AssetUtils/CreateSkeletalMeshUtil.h
+++ b/Engine/Plugins/Runtime/MeshModelingToolset/Source/ModelingComponentsEditorOnly/Public/AssetUtils/CreateSkeletalMeshUtil.h
@@ -33,7 +33,7 @@ namespace UE::AssetUtils
 	 * Only one of the arrays should be initialized, and it's size should be equal to 
 	 * the number of LODs specified for the Asset.
 	 */
-	struct MODELINGCOMPONENTSEDITORONLY_API FSkeletalMeshAssetMeshes
+	struct FSkeletalMeshAssetMeshes
 	{
 		TArray<const FDynamicMesh3*> DynamicMeshes;
 		TArray<const FMeshDescription*> MeshDescriptions;
@@ -43,7 +43,7 @@ namespace UE::AssetUtils
 	/**
 	 * Options for new USkeletalMesh asset created by Create() functions below.
 	 */
-	struct MODELINGCOMPONENTSEDITORONLY_API FSkeletalMeshAssetOptions
+	struct FSkeletalMeshAssetOptions
 	{
 		// This package will be used if it is not nullptr, otherwise a new package will be created at NewAssetPath
 		UPackage* UsePackage = nullptr;
@@ -79,7 +79,7 @@ namespace UE::AssetUtils
 	 * Output information about a newly-created StaticMesh, returned by Create functions below.
 	 * Some fields may be null, if the relevant function did not create that type of object
 	 */
-	struct MODELINGCOMPONENTSEDITORONLY_API FSkeletalMeshResults
+	struct FSkeletalMeshResults
 	{
 		/** USkeletalMesh asset that was created */
 		USkeletalMesh* SkeletalMesh = nullptr;
diff --git a/Engine/Plugins/Runtime/Metasound/Source/MetasoundFrontend/Private/Analysis/MetasoundFrontendVertexAnalyzerTriggerDensity.cpp b/Engine/Plugins/Runtime/Metasound/Source/MetasoundFrontend/Private/Analysis/MetasoundFrontendVertexAnalyzerTriggerDensity.cpp
--- a/Engine/Plugins/Runtime/Metasound/Source/MetasoundFrontend/Private/Analysis/MetasoundFrontendVertexAnalyzerTriggerDensity.cpp
+++ b/Engine/Plugins/Runtime/Metasound/Source/MetasoundFrontend/Private/Analysis/MetasoundFrontendVertexAnalyzerTriggerDensity.cpp
@@ -45,6 +45,8 @@ namespace Metasound
 			FVertexAnalyzerBase::BindOutputData<float>(FOutputs::GetValue().Name, InParams.OperatorSettings, TDataReadReference<float>(EnvelopeValue));
 		}
 
+		FVertexAnalyzerTriggerDensity::~FVertexAnalyzerTriggerDensity() = default;
+
 		void FVertexAnalyzerTriggerDensity::Execute()
 		{
 			const FTrigger& Trigger = GetVertexData<FTrigger>();
diff --git a/Engine/Plugins/Runtime/Metasound/Source/MetasoundFrontend/Public/Analysis/MetasoundFrontendVertexAnalyzerEnvelopeFollower.h b/Engine/Plugins/Runtime/Metasound/Source/MetasoundFrontend/Public/Analysis/MetasoundFrontendVertexAnalyzerEnvelopeFollower.h
--- a/Engine/Plugins/Runtime/Metasound/Source/MetasoundFrontend/Public/Analysis/MetasoundFrontendVertexAnalyzerEnvelopeFollower.h
+++ b/Engine/Plugins/Runtime/Metasound/Source/MetasoundFrontend/Public/Analysis/MetasoundFrontendVertexAnalyzerEnvelopeFollower.h
@@ -6,6 +6,7 @@
 #include "Containers/Array.h"
 #include "DSP/EnvelopeFollower.h"
 #include "MetasoundDataReferenceCollection.h"
+#include "MetasoundPrimitives.h"
 #include "MetasoundRouter.h"
 
 
diff --git a/Engine/Plugins/Runtime/Metasound/Source/MetasoundFrontend/Public/Analysis/MetasoundFrontendVertexAnalyzerTriggerDensity.h b/Engine/Plugins/Runtime/Metasound/Source/MetasoundFrontend/Public/Analysis/MetasoundFrontendVertexAnalyzerTriggerDensity.h
--- a/Engine/Plugins/Runtime/Metasound/Source/MetasoundFrontend/Public/Analysis/MetasoundFrontendVertexAnalyzerTriggerDensity.h
+++ b/Engine/Plugins/Runtime/Metasound/Source/MetasoundFrontend/Public/Analysis/MetasoundFrontendVertexAnalyzerTriggerDensity.h
@@ -5,6 +5,7 @@
 #include "Containers/Array.h"
 #include "DSP/BufferVectorOperations.h"
 #include "DSP/EnvelopeFollower.h"
+#include "MetasoundPrimitives.h"
 #include "MetasoundTrigger.h"
 
 
@@ -12,19 +13,18 @@ namespace Metasound
 {
 	namespace Frontend
 	{
-		class METASOUNDFRONTEND_API FVertexAnalyzerTriggerDensity : public FVertexAnalyzerBase
+		class FVertexAnalyzerTriggerDensity : public FVertexAnalyzerBase
 		{
 		public:
-			static const FName& GetAnalyzerName();
-			static const FName& GetDataType();
+			METASOUNDFRONTEND_API static const FName& GetAnalyzerName();
+			METASOUNDFRONTEND_API static const FName& GetDataType();
 
-			struct METASOUNDFRONTEND_API FOutputs
+			struct FOutputs
 			{
-				static const FAnalyzerOutput& GetValue();
-
+				METASOUNDFRONTEND_API static const FAnalyzerOutput& GetValue();
 			};
 
-			class METASOUNDFRONTEND_API FFactory : public TVertexAnalyzerFactory<FVertexAnalyzerTriggerDensity>
+			class FFactory : public TVertexAnalyzerFactory<FVertexAnalyzerTriggerDensity>
 			{
 			public:
 				virtual const TArray<FAnalyzerOutput>& GetAnalyzerOutputs() const override
@@ -34,10 +34,10 @@ namespace Metasound
 				}
 			};
 
-			FVertexAnalyzerTriggerDensity(const FCreateAnalyzerParams& InParams);
-			virtual ~FVertexAnalyzerTriggerDensity() = default;
+			METASOUNDFRONTEND_API FVertexAnalyzerTriggerDensity(const FCreateAnalyzerParams& InParams);
+			METASOUNDFRONTEND_API virtual ~FVertexAnalyzerTriggerDensity();
 
-			virtual void Execute() override;
+			METASOUNDFRONTEND_API virtual void Execute() override;
 
 		private:
 
diff --git a/Engine/Plugins/Runtime/StateTree/Source/StateTreeEditorModule/Private/Debugger/SStateTreeDebuggerViewRow.h b/Engine/Plugins/Runtime/StateTree/Source/StateTreeEditorModule/Private/Debugger/SStateTreeDebuggerViewRow.h
--- a/Engine/Plugins/Runtime/StateTree/Source/StateTreeEditorModule/Private/Debugger/SStateTreeDebuggerViewRow.h
+++ b/Engine/Plugins/Runtime/StateTree/Source/StateTreeEditorModule/Private/Debugger/SStateTreeDebuggerViewRow.h
@@ -5,6 +5,7 @@
 #if WITH_STATETREE_DEBUGGER
 
 #include "Debugger/StateTreeTraceTypes.h"
+#include "StateTree.h"
 #include "Templates/SharedPointer.h"
 #include "TraceServices/Model/Frames.h"
 #include "Widgets/Views/SListView.h"
diff --git a/Engine/Plugins/Runtime/Synthesis/Source/Synthesis/Classes/SourceEffects/SourceEffectEQ.h b/Engine/Plugins/Runtime/Synthesis/Source/Synthesis/Classes/SourceEffects/SourceEffectEQ.h
--- a/Engine/Plugins/Runtime/Synthesis/Source/Synthesis/Classes/SourceEffects/SourceEffectEQ.h
+++ b/Engine/Plugins/Runtime/Synthesis/Source/Synthesis/Classes/SourceEffects/SourceEffectEQ.h
@@ -49,19 +49,19 @@ struct SYNTHESIS_API FSourceEffectEQSettings
 	TArray<FSourceEffectEQBand> EQBands;
 };
 
-class SYNTHESIS_API FSourceEffectEQ : public FSoundEffectSource
+class FSourceEffectEQ : public FSoundEffectSource
 {
 public:
-	FSourceEffectEQ();
+	SYNTHESIS_API FSourceEffectEQ();
 
 	// Called on an audio effect at initialization on main thread before audio processing begins.
-	virtual void Init(const FSoundEffectSourceInitData& InitData) override;
+	SYNTHESIS_API virtual void Init(const FSoundEffectSourceInitData& InitData) override;
 	
 	// Called when an audio effect preset is changed
-	virtual void OnPresetChanged() override;
+	SYNTHESIS_API virtual void OnPresetChanged() override;
 
 	// Process the input block of audio. Called on audio thread.
-	virtual void ProcessAudio(const FSoundEffectSourceInputData& InData, float* OutAudioBufferData) override;
+	SYNTHESIS_API virtual void ProcessAudio(const FSoundEffectSourceInputData& InData, float* OutAudioBufferData) override;
 
 protected:
 
diff --git a/Engine/Plugins/Runtime/Synthesis/Source/Synthesis/Private/ConvolutionReverb.cpp b/Engine/Plugins/Runtime/Synthesis/Source/Synthesis/Private/ConvolutionReverb.cpp
--- a/Engine/Plugins/Runtime/Synthesis/Source/Synthesis/Private/ConvolutionReverb.cpp
+++ b/Engine/Plugins/Runtime/Synthesis/Source/Synthesis/Private/ConvolutionReverb.cpp
@@ -440,6 +440,8 @@ namespace Audio
 
 	}
 
+	FConvolutionReverb::~FConvolutionReverb() = default;
+
 	void FConvolutionReverb::SetSettings(const FConvolutionReverbSettings& InSettings)
 	{
 		Settings = InSettings;
diff --git a/Engine/Plugins/Runtime/Synthesis/Source/Synthesis/Public/ConvolutionReverb.h b/Engine/Plugins/Runtime/Synthesis/Source/Synthesis/Public/ConvolutionReverb.h
--- a/Engine/Plugins/Runtime/Synthesis/Source/Synthesis/Public/ConvolutionReverb.h
+++ b/Engine/Plugins/Runtime/Synthesis/Source/Synthesis/Public/ConvolutionReverb.h
@@ -137,7 +137,7 @@ namespace Audio
 				  +------------+
 
 	*/
-	class SYNTHESIS_API FConvolutionReverb 
+	class FConvolutionReverb 
 	{
 		FConvolutionReverb() = delete;
 		FConvolutionReverb(const FConvolutionReverb&) = delete;
@@ -146,6 +146,7 @@ namespace Audio
 
 
 	public:
+		~FConvolutionReverb();
 
 		// Create a convolution reverb object. This performs creation of the convolution algorithm object,
 		// converting sample rates of impulse responses, sets the impulse response and initializes the
@@ -155,22 +156,22 @@ namespace Audio
 		// @params InSettings - The initial settings for the convolution reverb.
 		//
 		// @return TUniquePtr<Audio::IConvolutionAlgorithm>  Will be invalid if there was an error.
-		static TUniquePtr<FConvolutionReverb> CreateConvolutionReverb(const FConvolutionReverbInitData& InInitData, const FConvolutionReverbSettings& InSettings=FConvolutionReverbSettings::DefaultSettings);
+		SYNTHESIS_API static TUniquePtr<FConvolutionReverb> CreateConvolutionReverb(const FConvolutionReverbInitData& InInitData, const FConvolutionReverbSettings& InSettings=FConvolutionReverbSettings::DefaultSettings);
 		
-		void SetSettings(const FConvolutionReverbSettings& InSettings);
+		SYNTHESIS_API void SetSettings(const FConvolutionReverbSettings& InSettings);
 
-		const FConvolutionReverbSettings& GetSettings() const;
+		SYNTHESIS_API const FConvolutionReverbSettings& GetSettings() const;
 
 
 		// If the number of input frames changes between callbacks, the output may contain discontinuities.
-		void ProcessAudio(int32 InNumInputChannels, const FAlignedFloatBuffer& InputAudio, int32 InNumOutputChannels, FAlignedFloatBuffer& OutputAudio);
-		void ProcessAudio(int32 InNumInputChannels, const float* InputAudio, int32 InNumOutputChannels, float* OutputAudio, const int32 InNumFrames);
+		SYNTHESIS_API void ProcessAudio(int32 InNumInputChannels, const FAlignedFloatBuffer& InputAudio, int32 InNumOutputChannels, FAlignedFloatBuffer& OutputAudio);
+		SYNTHESIS_API void ProcessAudio(int32 InNumInputChannels, const float* InputAudio, int32 InNumOutputChannels, float* OutputAudio, const int32 InNumFrames);
 
-		int32 GetNumInputChannels() const;
-		int32 GetNumOutputChannels() const;
+		SYNTHESIS_API int32 GetNumInputChannels() const;
+		SYNTHESIS_API int32 GetNumOutputChannels() const;
 
-		static void InterleaveBuffer(FAlignedFloatBuffer& OutBuffer, const TArray<FAlignedFloatBuffer>& InputBuffers, const int32 NumChannels);
-		static void DeinterleaveBuffer(TArray<FAlignedFloatBuffer>& OutputBuffers, TArrayView<const float> InputBuffer, const int32 NumChannels);
+		SYNTHESIS_API static void InterleaveBuffer(FAlignedFloatBuffer& OutBuffer, const TArray<FAlignedFloatBuffer>& InputBuffers, const int32 NumChannels);
+		SYNTHESIS_API static void DeinterleaveBuffer(TArray<FAlignedFloatBuffer>& OutputBuffers, TArrayView<const float> InputBuffer, const int32 NumChannels);
 
 	private:
 
diff --git a/Engine/Plugins/Runtime/Synthesis/Source/Synthesis/Public/Flexiverb.h b/Engine/Plugins/Runtime/Synthesis/Source/Synthesis/Public/Flexiverb.h
--- a/Engine/Plugins/Runtime/Synthesis/Source/Synthesis/Public/Flexiverb.h
+++ b/Engine/Plugins/Runtime/Synthesis/Source/Synthesis/Public/Flexiverb.h
@@ -46,20 +46,20 @@ namespace Audio
 	/* https://ccrma.stanford.edu/~jos/pasp/Householder_Feedback_Matrix.html        */
 	/********************************************************************************/
 
-	class SYNTHESIS_API FFlexiverb
+	class FFlexiverb
 	{
 	public:
-		FFlexiverb();
-		~FFlexiverb();
+		SYNTHESIS_API FFlexiverb();
+		SYNTHESIS_API ~FFlexiverb();
 
 		// Initialize the reverb with the given sample rate and initial settings.
-		void Init(const int32 InSampleRate, const FFlexiverbSettings Settings = FFlexiverbSettings());
+		SYNTHESIS_API void Init(const int32 InSampleRate, const FFlexiverbSettings Settings = FFlexiverbSettings());
 
 		// Updates the current settings of this reverb.
-		void SetSettings(const FFlexiverbSettings& InSettings);
+		SYNTHESIS_API void SetSettings(const FFlexiverbSettings& InSettings);
 
 		// Process a single audio frame.
-		void ProcessAudioFrame(const float* InBuffer, const int32 InChannels, float* OutBuffer, const int32 OutChannels);
+		SYNTHESIS_API void ProcessAudioFrame(const float* InBuffer, const int32 InChannels, float* OutBuffer, const int32 OutChannels);
 
 	private:
 		// Calculates coefficients based on the current settings.
diff --git a/Engine/Plugins/Runtime/Windows/WinDualShock/Source/WinDualShock/Private/WinDualShockControllers.h b/Engine/Plugins/Runtime/Windows/WinDualShock/Source/WinDualShock/Private/WinDualShockControllers.h
--- a/Engine/Plugins/Runtime/Windows/WinDualShock/Source/WinDualShock/Private/WinDualShockControllers.h
+++ b/Engine/Plugins/Runtime/Windows/WinDualShock/Source/WinDualShock/Private/WinDualShockControllers.h
@@ -2,9 +2,15 @@
 
 #pragma once
 
+#if DUALSHOCK4_SUPPORT
+
 #include "WinDualShock.h"
-#include <pad.h>
-#include <pad_audio.h>
+#include LIBSCEPAD_PLATFORM_INCLUDE
+
+THIRD_PARTY_INCLUDES_START
+	#include <pad.h>
+	#include <pad_audio.h>
+THIRD_PARTY_INCLUDES_END
 
 class FWinDualShockControllers : public FPlatformControllers
 {
@@ -49,3 +55,5 @@ public:
 private:
 	float OutputGain = 1.0f;
 };
+
+#endif // DUALSHOCK4_SUPPORT
diff --git a/Engine/Plugins/Runtime/Windows/WinDualShock/Source/WinDualShock/Private/WinDualShockExternalEndpoints.h b/Engine/Plugins/Runtime/Windows/WinDualShock/Source/WinDualShock/Private/WinDualShockExternalEndpoints.h
--- a/Engine/Plugins/Runtime/Windows/WinDualShock/Source/WinDualShock/Private/WinDualShockExternalEndpoints.h
+++ b/Engine/Plugins/Runtime/Windows/WinDualShock/Source/WinDualShock/Private/WinDualShockExternalEndpoints.h
@@ -2,8 +2,11 @@
 
 #pragma once
 
+#if DUALSHOCK4_SUPPORT
+
 #include "AudioMixer.h"
 #include "IAudioEndpoint.h"
+#include "WinDualShock.h"
 #include "WinDualShockSettings.h"
 
 /**
@@ -81,3 +84,5 @@ protected:
 		return true;
 	}
 };
+
+#endif // DUALSHOCK4_SUPPORT
diff --git a/Engine/Plugins/WorldMetrics/Source/WorldMetricsTest/Public/WorldMetricsTestUtil.h b/Engine/Plugins/WorldMetrics/Source/WorldMetricsTest/Public/WorldMetricsTestUtil.h
--- a/Engine/Plugins/WorldMetrics/Source/WorldMetricsTest/Public/WorldMetricsTestUtil.h
+++ b/Engine/Plugins/WorldMetrics/Source/WorldMetricsTest/Public/WorldMetricsTestUtil.h
@@ -2,6 +2,8 @@
 
 #pragma once
 
+#include "Containers/Array.h"
+
 #ifndef REQUIRE_MESSAGE
 #define REQUIRE_MESSAGE(Message, Expr)           \
 	if (!(Expr))                                 \
diff --git a/Engine/Source/Developer/AutomationWindow/Private/SAutomationTestItem.h b/Engine/Source/Developer/AutomationWindow/Private/SAutomationTestItem.h
--- a/Engine/Source/Developer/AutomationWindow/Private/SAutomationTestItem.h
+++ b/Engine/Source/Developer/AutomationWindow/Private/SAutomationTestItem.h
@@ -8,6 +8,7 @@
 #include "Widgets/DeclarativeSyntaxSupport.h"
 #include "Styling/SlateColor.h"
 #include "Widgets/SWidget.h"
+#include "Widgets/Views/SListView.h"
 #include "Widgets/Views/STableViewBase.h"
 #include "Widgets/Views/STableRow.h"
 
diff --git a/Engine/Source/Developer/Horde/Private/Storage/Nodes/ChunkNode.cpp b/Engine/Source/Developer/Horde/Private/Storage/Nodes/ChunkNode.cpp
--- a/Engine/Source/Developer/Horde/Private/Storage/Nodes/ChunkNode.cpp
+++ b/Engine/Source/Developer/Horde/Private/Storage/Nodes/ChunkNode.cpp
@@ -74,17 +74,6 @@ FBlobHandleWithHash FChunkNode::Write(FBlobWriter& Writer, const TArrayView<cons
 
 // ----------------------------------------------------------------------
 
-struct FChunkNodeReader::FStackEntry
-{
-	FBlob Blob;
-	size_t Position;
-
-	FStackEntry(FBlob InBlob)
-		: Blob(MoveTemp(InBlob))
-		, Position(0)
-	{ }
-};
-
 FChunkNodeReader::FChunkNodeReader(FBlob Blob)
 {
 	Stack.Add(FStackEntry(MoveTemp(Blob)));
diff --git a/Engine/Source/Developer/Horde/Public/Storage/Nodes/ChunkNode.h b/Engine/Source/Developer/Horde/Public/Storage/Nodes/ChunkNode.h
--- a/Engine/Source/Developer/Horde/Public/Storage/Nodes/ChunkNode.h
+++ b/Engine/Source/Developer/Horde/Public/Storage/Nodes/ChunkNode.h
@@ -62,7 +62,16 @@ public:
 	operator bool() const;
 
 private:
-	struct FStackEntry;
+	struct FStackEntry
+	{
+		FBlob Blob;
+		size_t Position;
+
+		FStackEntry(FBlob InBlob)
+			: Blob(MoveTemp(InBlob))
+			, Position(0)
+		{ }
+	};
 	TArray<FStackEntry> Stack;
 };
 
diff --git a/Engine/Source/Developer/MaterialBaking/Private/MaterialBakingModule.h b/Engine/Source/Developer/MaterialBaking/Private/MaterialBakingModule.h
--- a/Engine/Source/Developer/MaterialBaking/Private/MaterialBakingModule.h
+++ b/Engine/Source/Developer/MaterialBaking/Private/MaterialBakingModule.h
@@ -2,6 +2,7 @@
 
 #pragma once
 
+#include "Engine/TextureRenderTarget2D.h"
 #include "UObject/StrongObjectPtr.h"
 #include "Modules/ModuleInterface.h"
 #include "PixelFormat.h"
diff --git a/Engine/Source/Developer/MessageLog/Private/UserInterface/SMessageLogMessageListRow.h b/Engine/Source/Developer/MessageLog/Private/UserInterface/SMessageLogMessageListRow.h
--- a/Engine/Source/Developer/MessageLog/Private/UserInterface/SMessageLogMessageListRow.h
+++ b/Engine/Source/Developer/MessageLog/Private/UserInterface/SMessageLogMessageListRow.h
@@ -11,6 +11,7 @@
 #include "Widgets/SWidget.h"
 #include "Widgets/Views/STableViewBase.h"
 #include "Widgets/SBoxPanel.h"
+#include "Widgets/Views/SListView.h"
 #include "Widgets/Views/STableRow.h"
 
 #if WITH_EDITOR
diff --git a/Engine/Source/Developer/ScreenShotComparison/Private/Models/ScreenShotComparisonFilter.h b/Engine/Source/Developer/ScreenShotComparison/Private/Models/ScreenShotComparisonFilter.h
deleted file mode 100644
diff --git a/Engine/Source/Developer/ScreenShotComparison/Private/Widgets/SScreenComparisonRow.h b/Engine/Source/Developer/ScreenShotComparison/Private/Widgets/SScreenComparisonRow.h
--- a/Engine/Source/Developer/ScreenShotComparison/Private/Widgets/SScreenComparisonRow.h
+++ b/Engine/Source/Developer/ScreenShotComparison/Private/Widgets/SScreenComparisonRow.h
@@ -6,6 +6,7 @@
 #include "Widgets/DeclarativeSyntaxSupport.h"
 #include "Input/Reply.h"
 #include "Widgets/SWidget.h"
+#include "Widgets/Views/SListView.h"
 #include "Widgets/Views/STableViewBase.h"
 #include "Widgets/Views/STableRow.h"
 #include "ISourceControlState.h"
diff --git a/Engine/Source/Developer/Windows/LiveCoding/Private/External/LC_VirtualMemoryTypes.h b/Engine/Source/Developer/Windows/LiveCoding/Private/External/LC_VirtualMemoryTypes.h
--- a/Engine/Source/Developer/Windows/LiveCoding/Private/External/LC_VirtualMemoryTypes.h
+++ b/Engine/Source/Developer/Windows/LiveCoding/Private/External/LC_VirtualMemoryTypes.h
@@ -2,6 +2,9 @@
 
 #pragma once
 
+// BEGIN EPIC MOD
+#include "Windows/WindowsHWrapper.h"
+// END EPIC MOD
 
 namespace VirtualMemory
 {
diff --git a/Engine/Source/Editor/DataTableEditor/Private/SDataTableListViewRow.h b/Engine/Source/Editor/DataTableEditor/Private/SDataTableListViewRow.h
--- a/Engine/Source/Editor/DataTableEditor/Private/SDataTableListViewRow.h
+++ b/Engine/Source/Editor/DataTableEditor/Private/SDataTableListViewRow.h
@@ -19,6 +19,7 @@
 #include "UObject/NameTypes.h"
 #include "Widgets/DeclarativeSyntaxSupport.h"
 #include "Widgets/SCompoundWidget.h"
+#include "Widgets/Views/SListView.h"
 #include "Widgets/Views/STableRow.h"
 
 class FDataTableEditor;
diff --git a/Engine/Source/Editor/FoliageEdit/Private/FoliagePaletteItem.h b/Engine/Source/Editor/FoliageEdit/Private/FoliagePaletteItem.h
--- a/Engine/Source/Editor/FoliageEdit/Private/FoliagePaletteItem.h
+++ b/Engine/Source/Editor/FoliageEdit/Private/FoliagePaletteItem.h
@@ -14,6 +14,7 @@
 #include "Templates/UnrealTemplate.h"
 #include "UObject/NameTypes.h"
 #include "Widgets/DeclarativeSyntaxSupport.h"
+#include "Widgets/Views/SListView.h"
 #include "Widgets/Views/STableRow.h"
 
 class FEdModeFoliage;
diff --git a/Engine/Source/Editor/GameProjectGeneration/Private/GameProjectUtils.cpp b/Engine/Source/Editor/GameProjectGeneration/Private/GameProjectUtils.cpp
--- a/Engine/Source/Editor/GameProjectGeneration/Private/GameProjectUtils.cpp
+++ b/Engine/Source/Editor/GameProjectGeneration/Private/GameProjectUtils.cpp
@@ -2082,31 +2082,6 @@ void GameProjectUtils::AddHardwareConfigValues(const FProjectInformation& InProj
 			}
 		}
 	}
-
-	// Don't override these settings for templates
-	if (InProjectInfo.TemplateFile.IsEmpty())
-	{
-		// New projects always have DX12 by default on Windows
-		ConfigValues.Emplace(TEXT("DefaultEngine.ini"),
-			TEXT("/Script/WindowsTargetPlatform.WindowsTargetSettings"),
-			TEXT("DefaultGraphicsRHI"),
-			TEXT("DefaultGraphicsRHI_DX12"),
-			false /* ShouldReplaceExistingValue */);
-
-		// Force clear D3D12TargetedShaderFormats since the BaseEngine list can change at any time.
-		ConfigValues.Emplace(TEXT("DefaultEngine.ini"),
-			TEXT("/Script/WindowsTargetPlatform.WindowsTargetSettings"),
-			TEXT("!D3D12TargetedShaderFormats"),
-			TEXT("ClearArray"),
-			false /* ShouldReplaceExistingValue */);
-
-		// New projects always have DX12 only supporting SM6 by default on Windows
-		ConfigValues.Emplace(TEXT("DefaultEngine.ini"),
-			TEXT("/Script/WindowsTargetPlatform.WindowsTargetSettings"),
-			TEXT("+D3D12TargetedShaderFormats"),
-			TEXT("PCD3D_SM6"),
-			false /* ShouldReplaceExistingValue */);
-	}
 }
 
 bool GameProjectUtils::GenerateConfigFiles(const FProjectInformation& InProjectInfo, TArray<FString>& OutCreatedFiles, FText& OutFailReason, FGuid& OutProjectID)
diff --git a/Engine/Source/Editor/HierarchicalLODOutliner/Private/HLODTreeWidgetItem.h b/Engine/Source/Editor/HierarchicalLODOutliner/Private/HLODTreeWidgetItem.h
--- a/Engine/Source/Editor/HierarchicalLODOutliner/Private/HLODTreeWidgetItem.h
+++ b/Engine/Source/Editor/HierarchicalLODOutliner/Private/HLODTreeWidgetItem.h
@@ -11,6 +11,7 @@
 #include "Templates/UnrealTemplate.h"
 #include "UObject/NameTypes.h"
 #include "Widgets/DeclarativeSyntaxSupport.h"
+#include "Widgets/Views/SListView.h"
 #include "Widgets/Views/STableRow.h"
 
 class FDragDropEvent;
diff --git a/Engine/Source/Editor/Layers/Private/SLayersViewRow.h b/Engine/Source/Editor/Layers/Private/SLayersViewRow.h
--- a/Engine/Source/Editor/Layers/Private/SLayersViewRow.h
+++ b/Engine/Source/Editor/Layers/Private/SLayersViewRow.h
@@ -9,6 +9,7 @@
 #include "Widgets/SWidget.h"
 #include "Widgets/DeclarativeSyntaxSupport.h"
 #include "LayerViewModel.h"
+#include "Widgets/Views/SListView.h"
 #include "Widgets/Views/STableViewBase.h"
 #include "Widgets/Views/STableRow.h"
 
diff --git a/Engine/Source/Editor/LocalizationDashboard/Private/SLocalizationDashboardTargetRow.h b/Engine/Source/Editor/LocalizationDashboard/Private/SLocalizationDashboardTargetRow.h
--- a/Engine/Source/Editor/LocalizationDashboard/Private/SLocalizationDashboardTargetRow.h
+++ b/Engine/Source/Editor/LocalizationDashboard/Private/SLocalizationDashboardTargetRow.h
@@ -10,6 +10,7 @@
 #include "Templates/SharedPointer.h"
 #include "Templates/UnrealTemplate.h"
 #include "UObject/NameTypes.h"
+#include "Widgets/Views/SListView.h"
 #include "Widgets/Views/STableRow.h"
 
 class IPropertyHandle;
diff --git a/Engine/Source/Editor/LocalizationDashboard/Private/SLocalizationTargetEditorCultureRow.h b/Engine/Source/Editor/LocalizationDashboard/Private/SLocalizationTargetEditorCultureRow.h
--- a/Engine/Source/Editor/LocalizationDashboard/Private/SLocalizationTargetEditorCultureRow.h
+++ b/Engine/Source/Editor/LocalizationDashboard/Private/SLocalizationTargetEditorCultureRow.h
@@ -4,6 +4,7 @@
 #include "CoreMinimal.h"
 #include "Input/Reply.h"
 #include "Widgets/SWidget.h"
+#include "Widgets/Views/SListView.h"
 #include "Widgets/Views/STableViewBase.h"
 #include "Widgets/Views/STableRow.h"
 #include "PropertyHandle.h"
diff --git a/Engine/Source/Editor/Persona/Private/AnimGraphNodeDetails.h b/Engine/Source/Editor/Persona/Private/AnimGraphNodeDetails.h
--- a/Engine/Source/Editor/Persona/Private/AnimGraphNodeDetails.h
+++ b/Engine/Source/Editor/Persona/Private/AnimGraphNodeDetails.h
@@ -10,6 +10,7 @@
 #include "IPropertyTypeCustomization.h"
 #include "PropertyHandle.h"
 #include "IDetailCustomization.h"
+#include "Widgets/Views/SListView.h"
 #include "Widgets/Views/STableViewBase.h"
 #include "Widgets/Views/STableRow.h"
 #include "IDetailPropertyExtensionHandler.h"
diff --git a/Engine/Source/Editor/Persona/Public/SAnimAttributeView.h b/Engine/Source/Editor/Persona/Public/SAnimAttributeView.h
--- a/Engine/Source/Editor/Persona/Public/SAnimAttributeView.h
+++ b/Engine/Source/Editor/Persona/Public/SAnimAttributeView.h
@@ -4,6 +4,7 @@
 
 #include "CoreMinimal.h"
 #include "IStructureDetailsView.h"
+#include "Widgets/Views/SListView.h"
 #include "Widgets/Views/STableRow.h"
 #include "Widgets/SCompoundWidget.h"
 #include "Animation/AttributesRuntime.h"
diff --git a/Engine/Source/Editor/PropertyEditor/Private/UserInterface/PropertyEditor/SPropertyEditorTableRow.h b/Engine/Source/Editor/PropertyEditor/Private/UserInterface/PropertyEditor/SPropertyEditorTableRow.h
--- a/Engine/Source/Editor/PropertyEditor/Private/UserInterface/PropertyEditor/SPropertyEditorTableRow.h
+++ b/Engine/Source/Editor/PropertyEditor/Private/UserInterface/PropertyEditor/SPropertyEditorTableRow.h
@@ -6,6 +6,7 @@
 #include "Input/Reply.h"
 #include "Widgets/DeclarativeSyntaxSupport.h"
 #include "Widgets/SWidget.h"
+#include "Widgets/Views/SListView.h"
 #include "Widgets/Views/STableViewBase.h"
 #include "Widgets/Views/STableRow.h"
 #include "IPropertyUtilities.h"
diff --git a/Engine/Source/Editor/SkeletonEditor/Private/SSkeletonTreeRow.h b/Engine/Source/Editor/SkeletonEditor/Private/SSkeletonTreeRow.h
--- a/Engine/Source/Editor/SkeletonEditor/Private/SSkeletonTreeRow.h
+++ b/Engine/Source/Editor/SkeletonEditor/Private/SSkeletonTreeRow.h
@@ -15,6 +15,7 @@
 #include "Templates/UnrealTemplate.h"
 #include "UObject/NameTypes.h"
 #include "Widgets/DeclarativeSyntaxSupport.h"
+#include "Widgets/Views/SListView.h"
 #include "Widgets/Views/STableRow.h"
 
 class FDragDropEvent;
diff --git a/Engine/Source/Editor/WorldBrowser/Private/SWorldHierarchyItem.h b/Engine/Source/Editor/WorldBrowser/Private/SWorldHierarchyItem.h
--- a/Engine/Source/Editor/WorldBrowser/Private/SWorldHierarchyItem.h
+++ b/Engine/Source/Editor/WorldBrowser/Private/SWorldHierarchyItem.h
@@ -9,6 +9,7 @@
 #include "Fonts/SlateFontInfo.h"
 #include "Widgets/DeclarativeSyntaxSupport.h"
 #include "Widgets/SWidget.h"
+#include "Widgets/Views/SListView.h"
 #include "Widgets/Views/STableViewBase.h"
 #include "Widgets/Views/STableRow.h"
 
diff --git a/Engine/Source/Runtime/Engine/Private/PhysicsEngine/CollisionAnalyzerCapture.h b/Engine/Source/Runtime/Engine/Private/PhysicsEngine/CollisionAnalyzerCapture.h
--- a/Engine/Source/Runtime/Engine/Private/PhysicsEngine/CollisionAnalyzerCapture.h
+++ b/Engine/Source/Runtime/Engine/Private/PhysicsEngine/CollisionAnalyzerCapture.h
@@ -8,14 +8,13 @@
 
 #include "ICollisionAnalyzer.h"
 #include "CollisionAnalyzerModule.h"
+#include "Engine/HitResult.h"
 #include "Physics/PhysicsInterfaceCore.h"
 #include "Physics/PhysicsInterfaceUtils.h"
 
 #include "Chaos/Capsule.h"
 #include "Chaos/Sphere.h"
 
-struct FHitResult;
-
 extern bool bSkipCapture;
 extern bool GCollisionAnalyzerIsRecording;
 
diff --git a/Engine/Source/Runtime/Experimental/Chaos/Public/Chaos/PBDChainUpdateFromDeltaPosition.h b/Engine/Source/Runtime/Experimental/Chaos/Public/Chaos/PBDChainUpdateFromDeltaPosition.h
deleted file mode 100644
