Add deprecated PlatformCryptoOpenSSL module to warn licensees of the module's replacement.

#rnx

diff --git a/Engine/Plugins/Experimental/PlatformCrypto/Source/PlatformCryptoOpenSSL/PlatformCryptoOpenSSL.Build.cs b/Engine/Plugins/Experimental/PlatformCrypto/Source/PlatformCryptoOpenSSL/PlatformCryptoOpenSSL.Build.cs
new file mode 100644
--- /dev/null
+++ b/Engine/Plugins/Experimental/PlatformCrypto/Source/PlatformCryptoOpenSSL/PlatformCryptoOpenSSL.Build.cs
@@ -0,0 +1,14 @@
+// Copyright Epic Games, Inc. All Rights Reserved.
+
+using System;
+using System.IO;
+using UnrealBuildTool;
+
+[Obsolete("Deprecated in UE5.6 - Use PlatformCryptoContext instead.")]
+public class PlatformCryptoOpenSSL : ModuleRules
+{
+	public PlatformCryptoOpenSSL(ReadOnlyTargetRules Target) : base(Target)
+	{
+		Type = ModuleType.External;
+	}
+}
