Add ED25519 key pair generation and sign/verify support to Engine/PlatformCrypto

#rnx

- Integration of https://p4-swarm.epicgames.net/changes/36565531 https://p4-swarm.epicgames.net/changes/25037446 https://p4-swarm.epicgames.net/changes/33625021 https://p4-swarm.epicgames.net/changes/51089916;
- Added FCryptoKeyGeneratorOpenSSL and FCryptoKeyGeneratorSwitchOpenSSL to generate keys;
- Added FSignerEd25519OpenSSL and FVerifierEd25519OpenSSL to sign and verify data.

#Synced-CL 50588029

diff --git a/Engine/Plugins/Experimental/PlatformCrypto/Source/PlatformCryptoContext/Private/EncryptionContextOpenSSL.cpp b/Engine/Plugins/Experimental/PlatformCrypto/Source/PlatformCryptoContext/Private/EncryptionContextOpenSSL.cpp
--- a/Engine/Plugins/Experimental/PlatformCrypto/Source/PlatformCryptoContext/Private/EncryptionContextOpenSSL.cpp
+++ b/Engine/Plugins/Experimental/PlatformCrypto/Source/PlatformCryptoContext/Private/EncryptionContextOpenSSL.cpp
@@ -8,6 +8,7 @@
 #include "PlatformCryptoUtilsOpenSSL.h"
 
 #include "Containers/StringConv.h"
+#include "Misc/ScopeExit.h"
 
 THIRD_PARTY_INCLUDES_START
 #include <openssl/bn.h>
@@ -82,6 +83,204 @@ EPlatformCryptoResult FSHA256HasherOpenSSL::Finalize(const TArrayView<uint8> Out
 	return EPlatformCryptoResult::Success;
 }
 
+struct FEd25519CheckContextOpenSSL
+{
+	void Shutdown()
+	{
+		if (EvpPkey != nullptr)
+		{
+			EVP_PKEY_free(EvpPkey);
+			EvpPkey = nullptr;
+		}
+		if (EvpMdCtx != nullptr)
+		{
+			EVP_MD_CTX_free(EvpMdCtx);
+			EvpMdCtx = nullptr;
+		}
+
+		State = EState::NeedsInit;
+	}
+
+	bool IsReady() const
+	{
+		return State == EState::Ready;
+	}
+
+	EVP_MD_CTX* EvpMdCtx = nullptr;
+	EVP_PKEY* EvpPkey = nullptr;
+
+	enum class EState { NeedsInit, Ready, InitFailed } State = EState::NeedsInit;
+
+protected:
+	~FEd25519CheckContextOpenSSL()
+	{
+		Shutdown();
+	}
+};
+
+struct FSignerEd25519OpenSSL::FImplDetails : public FEd25519CheckContextOpenSSL
+{
+};
+
+FSignerEd25519OpenSSL::FSignerEd25519OpenSSL() : Inner(MakePimpl<FSignerEd25519OpenSSL::FImplDetails>())
+{
+	check(Inner);
+}
+
+FSignerEd25519OpenSSL::~FSignerEd25519OpenSSL()
+{
+}
+
+bool FSignerEd25519OpenSSL::IsReady() const
+{
+	return Inner->IsReady();
+}
+
+bool FSignerEd25519OpenSSL::SetPrivateKey(const TArrayView<const uint8> Key)
+{
+	Inner->Shutdown();
+
+	Inner->State = FImplDetails::EState::InitFailed;
+
+	Inner->EvpMdCtx = EVP_MD_CTX_new();
+	if (Inner->EvpMdCtx)
+	{
+		Inner->EvpPkey = EVP_PKEY_new_raw_private_key(EVP_PKEY_ED25519, NULL, Key.GetData(), Key.Num());
+		if (Inner->EvpPkey)
+		{
+			if (EVP_DigestSignInit(Inner->EvpMdCtx, NULL, NULL, NULL, Inner->EvpPkey) == 1)
+			{
+				Inner->State = FImplDetails::EState::Ready;
+			}
+		}
+	}
+
+	return Inner->IsReady();
+}
+
+bool FSignerEd25519OpenSSL::Sign(const TArrayView<const uint8> Message, const TArrayView<uint8> Signature)
+{
+	if (!Inner->IsReady() || Signature.Num() != OutputSignatureLength)
+	{
+		return false;
+	}
+
+	size_t SignatureLength = OutputSignatureLength;
+
+	bool Succeeded = EVP_DigestSign(Inner->EvpMdCtx, Signature.GetData(), &SignatureLength, Message.GetData(), Message.Num()) == 1;
+	if (!Succeeded)
+	{
+		Inner->State = FImplDetails::EState::InitFailed;
+	}
+
+	return Succeeded;
+}
+
+struct FVerifierEd25519OpenSSL::FImplDetails : public FEd25519CheckContextOpenSSL
+{
+};
+
+FVerifierEd25519OpenSSL::FVerifierEd25519OpenSSL() : Inner(MakePimpl<FVerifierEd25519OpenSSL::FImplDetails>())
+{
+	check(Inner);
+}
+
+FVerifierEd25519OpenSSL::~FVerifierEd25519OpenSSL()
+{
+}
+
+bool FVerifierEd25519OpenSSL::IsReady() const
+{
+	return Inner->IsReady();
+}
+
+bool FVerifierEd25519OpenSSL::SetPublicKey(const TArrayView<const uint8> Key)
+{
+	Inner->Shutdown();
+
+	Inner->State = FImplDetails::EState::InitFailed;
+
+	Inner->EvpMdCtx = EVP_MD_CTX_new();
+	if (Inner->EvpMdCtx)
+	{
+		Inner->EvpPkey = EVP_PKEY_new_raw_public_key(EVP_PKEY_ED25519, NULL, Key.GetData(), Key.Num());
+		if (Inner->EvpPkey)
+		{
+			if (EVP_DigestVerifyInit(Inner->EvpMdCtx, NULL, NULL, NULL, Inner->EvpPkey) == 1)
+			{
+				Inner->State = FImplDetails::EState::Ready;
+			}
+		}
+	}
+
+	return Inner->IsReady();
+}
+
+bool FVerifierEd25519OpenSSL::Verify(const TArrayView<const uint8> Message, const TArrayView<const uint8> Signature)
+{
+	if (!Inner->IsReady())
+	{
+		return false;
+	}
+
+	return EVP_DigestVerify(Inner->EvpMdCtx, Signature.GetData(), Signature.Num(), Message.GetData(), Message.Num()) == 1;
+}
+
+TOptional<FCryptoKeyPair> FCryptoKeyGeneratorOpenSSL::GenerateEd25519() const
+{
+	TOptional<FCryptoKeyPair> Result;
+
+	const FScopedEvpPKeyContext Context(EVP_PKEY_ED25519);
+	if (!Context.Get())
+	{
+		UE_LOG(LogPlatformCryptoOpenSSL, Error, TEXT("CryptoKeyGeneratorOpenSSL::GenerateEd25519. Unable to allocates public key algorithm context"));
+		return Result;
+	}
+
+	int ErrCode = EVP_PKEY_keygen_init(Context.Get());
+	if (ErrCode != 1)
+	{
+		UE_LOG(LogPlatformCryptoOpenSSL, Error, TEXT("CryptoKeyGeneratorOpenSSL::GenerateEd25519. Unable to initialize a public key algorithm context. ErrCode=[%d]"), ErrCode);
+		return Result;
+	}
+
+	EVP_PKEY* PvtKey = nullptr;
+	ON_SCOPE_EXIT
+	{
+		if (PvtKey != nullptr)
+		{
+			EVP_PKEY_free(PvtKey);
+		}
+	};
+	ErrCode = EVP_PKEY_keygen(Context.Get(), &PvtKey);
+	if (ErrCode != 1)
+	{
+		UE_LOG(LogPlatformCryptoOpenSSL, Error, TEXT("CryptoKeyGeneratorOpenSSL::GenerateEd25519. Unable to generate keys. ErrCode=[%d]"), ErrCode);
+		return Result;
+	}
+
+	FCryptoKeyPair KeyPair;
+	KeyPair.PrivateKey.SetNumZeroed(32);
+	size_t PrivateKeyLen = KeyPair.PrivateKey.Num();
+	ErrCode = EVP_PKEY_get_raw_private_key(PvtKey, KeyPair.PrivateKey.GetData(), &PrivateKeyLen);
+	if (ErrCode != 1)
+	{
+		UE_LOG(LogPlatformCryptoOpenSSL, Error, TEXT("CryptoKeyGeneratorOpenSSL::GenerateEd25519. Unable to get private key. ErrCode=[%d] PrivateKeyLen=[%llu]"), ErrCode, static_cast<uint64_t>(PrivateKeyLen));
+		return Result;
+	}
+
+	KeyPair.PublicKey.SetNumZeroed(32);
+	size_t PublicKeyLen = KeyPair.PublicKey.Num();
+	ErrCode = EVP_PKEY_get_raw_public_key(PvtKey, KeyPair.PublicKey.GetData(), &PublicKeyLen);
+	if (ErrCode != 1)
+	{
+		UE_LOG(LogPlatformCryptoOpenSSL, Error, TEXT("CryptoKeyGeneratorOpenSSL::GenerateEd25519. Unable to get public key. ErrCode=[%d] PublicKeyLen=[%llu]"), ErrCode, static_cast<uint64_t>(PublicKeyLen));
+		return Result;
+	}
+
+	return Result.Emplace(MoveTemp(KeyPair));
+}
+
 TArray<uint8> FEncryptionContextOpenSSL::Encrypt_AES_256_ECB(const TArrayView<const uint8> Plaintext, const TArrayView<const uint8> Key, EPlatformCryptoResult& OutResult)
 {
 	DECLARE_SCOPE_CYCLE_COUNTER(TEXT("OpenSSL AES256 ECB Encrypt"), STAT_OpenSSL_AES256_ECB_Encrypt, STATGROUP_PlatformCrypto);
diff --git a/Engine/Plugins/Experimental/PlatformCrypto/Source/PlatformCryptoContext/Private/PlatformCryptoOpenSSLTypes.h b/Engine/Plugins/Experimental/PlatformCrypto/Source/PlatformCryptoContext/Private/PlatformCryptoOpenSSLTypes.h
--- a/Engine/Plugins/Experimental/PlatformCrypto/Source/PlatformCryptoContext/Private/PlatformCryptoOpenSSLTypes.h
+++ b/Engine/Plugins/Experimental/PlatformCrypto/Source/PlatformCryptoContext/Private/PlatformCryptoOpenSSLTypes.h
@@ -90,3 +90,38 @@ public:
 private:
 	EVP_MD_CTX* Context;
 };
+
+/**
+ * RAII wrapper for the OpenSSL EVP public key algorithm context object.
+ */
+class FScopedEvpPKeyContext
+{
+public:
+	/**
+	 * Creates OpenSSL EVP public key algorithm context on construction.
+	 */
+	FScopedEvpPKeyContext(int KeyAlgorithm) :
+		Context(EVP_PKEY_CTX_new_id(KeyAlgorithm, nullptr))
+	{
+	}
+
+	/** Disable copying/assigning */
+	FScopedEvpPKeyContext(FScopedEvpPKeyContext&) = delete;
+	FScopedEvpPKeyContext& operator=(FScopedEvpPKeyContext&) = delete;
+
+	/**
+	 * Releases OpenSSL EVP public key algorithm context.
+	 */
+	~FScopedEvpPKeyContext()
+	{
+		EVP_PKEY_CTX_free(Context);
+	}
+
+	/**
+	 * Returns OpenSSL EVP public key algorithm context.
+	 */
+	EVP_PKEY_CTX* Get() const { return Context; }
+
+private:
+	EVP_PKEY_CTX* Context;
+};
\ No newline at end of file
diff --git a/Engine/Plugins/Experimental/PlatformCrypto/Source/PlatformCryptoContext/Public/EncryptionContextOpenSSL.h b/Engine/Plugins/Experimental/PlatformCrypto/Source/PlatformCryptoContext/Public/EncryptionContextOpenSSL.h
--- a/Engine/Plugins/Experimental/PlatformCrypto/Source/PlatformCryptoContext/Public/EncryptionContextOpenSSL.h
+++ b/Engine/Plugins/Experimental/PlatformCrypto/Source/PlatformCryptoContext/Public/EncryptionContextOpenSSL.h
@@ -10,6 +10,7 @@ class IPlatformCryptoDecryptor;
 class IPlatformCryptoEncryptor;
 enum class EPlatformCryptoResult;
 class IPlatformCryptoContext;
+struct FCryptoKeyPair;
 
 /** Implementation details for SHA256 computation using OpenSSL */
 struct FSHA256HasherOpenSSL final
@@ -53,6 +54,125 @@ private:
 };
 using FSHA256Hasher = FSHA256HasherOpenSSL;
 
+/** Implementation details for ed25519 signing using OpenSSL */
+class FSignerEd25519OpenSSL final
+{
+public:
+	FSignerEd25519OpenSSL();
+
+	FSignerEd25519OpenSSL(FSignerEd25519OpenSSL&&) = default;
+	FSignerEd25519OpenSSL(const FSignerEd25519OpenSSL&) = delete;
+	FSignerEd25519OpenSSL& operator=(FSignerEd25519OpenSSL&&) = default;
+	FSignerEd25519OpenSSL& operator=(const FSignerEd25519OpenSSL&) = delete;
+
+	~FSignerEd25519OpenSSL();
+
+	/**
+	* Set private key for signing
+	*
+	* @param Key Buffer pointing to the private key data
+	* @return True if initialized succesfully
+	*/
+	bool SetPrivateKey(const TArrayView<const uint8> Key);
+
+	/**
+	* Signing data block
+	*
+	* @param Message Buffer pointing to the data
+	* @param Signature Buffer pointing to the signature
+	* @return True if signed succesfully
+	*/
+	bool Sign(const TArrayView<const uint8> Message, const TArrayView<uint8> Signature);
+	
+	/**
+	* Check if signer are ready
+	*/
+	bool IsReady() const;
+
+	/**
+	 * The final signature length in bytes.
+	*/
+	static constexpr const uint32 OutputSignatureLength = 64;
+
+	/**
+	 * The key length in bytes.
+	*/
+	static constexpr const uint32 KeyBufferLength = 32;
+
+private:
+	friend class FEncryptionContextOpenSSL;
+
+	struct FImplDetails;
+	TPimplPtr<FImplDetails> Inner;
+};
+using FSignerEd25519 = FSignerEd25519OpenSSL;
+
+/** Implementation details for ed25519 verify signature using OpenSSL */
+class FVerifierEd25519OpenSSL final
+{
+public:
+	FVerifierEd25519OpenSSL();
+
+	FVerifierEd25519OpenSSL(FVerifierEd25519OpenSSL&&) = default;
+	FVerifierEd25519OpenSSL(const FVerifierEd25519OpenSSL&) = delete;
+	FVerifierEd25519OpenSSL& operator=(FVerifierEd25519OpenSSL&&) = default;
+	FVerifierEd25519OpenSSL& operator=(const FVerifierEd25519OpenSSL&) = delete;
+
+	~FVerifierEd25519OpenSSL();
+
+	/**
+	* Set public key for verify signature
+	*
+	* @param Key Buffer pointing to the public key data
+	* @return True if inited successfully
+	*/
+	bool SetPublicKey(const TArrayView<const uint8> Key);
+
+	/**
+	* Verify data block
+	*
+	* @param Message Buffer pointing to the data
+	* @param Signature Buffer pointing to the signature
+	* @return True if signature are correct
+	*/
+	bool Verify(const TArrayView<const uint8> Message, const TArrayView<const uint8> Signature);
+
+	/**
+	* Check if verifier are ready
+	*/
+	bool IsReady() const;
+
+	/**
+	* The signature length in bytes.
+	*/
+	static constexpr const uint32 SignatureLength = 64;
+
+	/**
+	 * The key length in bytes.
+	*/
+	static constexpr const uint32 KeyBufferLength = 32;
+
+private:
+	friend class FEncryptionContextOpenSSL;
+
+	struct FImplDetails;
+	TPimplPtr<FImplDetails> Inner;
+};
+using FVerifierEd25519 = FVerifierEd25519OpenSSL;
+
+/** Implementation of crypto key generator using OpenSSL. */
+class FCryptoKeyGeneratorOpenSSL final
+{
+public:
+	/**
+	* Generates Ed25519 key pair.
+	*
+	* @return The structure with public/private key pair when success.
+	*/
+	TOptional<FCryptoKeyPair> GenerateEd25519() const;
+};
+using FCryptoKeyGenerator = FCryptoKeyGeneratorOpenSSL;
+
 /**
  * Interface to certain cryptographic algorithms, using OpenSSL to implement them.
  */
diff --git a/Engine/Plugins/Experimental/PlatformCrypto/Source/PlatformCryptoTypes/Public/PlatformCryptoTypes.h b/Engine/Plugins/Experimental/PlatformCrypto/Source/PlatformCryptoTypes/Public/PlatformCryptoTypes.h
--- a/Engine/Plugins/Experimental/PlatformCrypto/Source/PlatformCryptoTypes/Public/PlatformCryptoTypes.h
+++ b/Engine/Plugins/Experimental/PlatformCrypto/Source/PlatformCryptoTypes/Public/PlatformCryptoTypes.h
@@ -245,4 +245,10 @@ protected:
 	int32 BlockSize = 0;
 };
 
+/** Structure with public/private key pair. */
+struct FCryptoKeyPair
+{
+	TArray<uint8> PrivateKey;
+	TArray<uint8> PublicKey;
+};
 #undef UE_API
